VR Editor:

diff --git a/Engine/Source/Editor/VREditor/VREditorWorldInteraction.h b/Engine/Source/Editor/VREditor/VREditorWorldInteraction.h
--- a/Engine/Source/Editor/VREditor/VREditorWorldInteraction.h
+++ b/Engine/Source/Editor/VREditor/VREditorWorldInteraction.h
@@ -2,237 +2,75 @@
 
 #pragma once
 
-#include "VREditorWorldInteractionTypes.h"
-
+#include "ViewportWorldInteraction.h"
+#include "VREditorWorldInteraction.generated.h"
 
 /**
  * VR Editor interaction with the 3D world
  */
-class FVREditorWorldInteraction : public FGCObject
+UCLASS()
+class UVREditorWorldInteraction : public UViewportWorldInteraction
 {
-public:
-
-	/** Default constructor for FVREditorWorldInteraction, which sets up safe defaults */
-	FVREditorWorldInteraction( class FVREditorMode& InitOwner );
-
-	/** Clean up before destruction */
-	virtual ~FVREditorWorldInteraction();
-
-	// FGCObject overrides
-	virtual void AddReferencedObjects( FReferenceCollector& Collector ) override;
+	GENERATED_UCLASS_BODY()
 
-	/** Called by VREditorMode to update us every frame */
-	void Tick( class FEditorViewportClient* ViewportClient, const float DeltaTime );
+public:
 
-	/** Called by VREditorMode to draw debug graphics every frame */
-	void Render( const class FSceneView* SceneView, class FViewport* Viewport, class FPrimitiveDrawInterface* PDI );
+	// UViewportWorldInteraction overrides
+	virtual void Shutdown() override;
+	virtual bool IsInteractableComponent( const UActorComponent* Component ) const override;
+	virtual void StopDragging( class UViewportInteractor* Interactor ) override;
 
-	/** Gets the transform gizmo actor, or returns null if we currently don't have one */
-	class ABaseTransformGizmo* GetTransformGizmoActor() const
+	/** Sets the owner of this system */
+	void SetOwner( class FVREditorMode* NewOwner )
 	{
-		return TransformGizmoActor;
+		Owner = NewOwner;
 	}
 
-	/** Deletes whatever is selected right now. */
-	void DeleteSelectedObjects();
-
-	/** Called by our owner right before a map is loaded or switching between Simulate and normal Editor, so we can clean
-	    up our spawned actors */
-	void CleanUpActorsBeforeMapChangeOrSimulate();
-
-	/** Returns true if we're actively teleporting */
-	bool IsTeleporting() const
+	/** Gets the owner of this system */
+	class FVREditorMode& GetOwner()
 	{
-		return bIsTeleporting;
+		return *Owner;
 	}
 
-	/** Creates a Transformable for each selected actor, wiping out the existing transformables */
-	void SetupTransformablesForSelectedActors();
+	/** Gets the owner of this system (const) */
+	const class FVREditorMode& GetOwner() const
+	{
+		return *Owner;
+	}
 
-	/** Gets the Y delta of the trackpad */
-	float GetTrackpadSlideDelta( const int32 HandIndex );
+		/** Snaps the selected objects to the ground */
+	void SnapSelectedActorsToGround();
 
 protected:
 
-	/** Called when the user presses a button on their motion controller device */
-	void OnVRAction( FEditorViewportClient& ViewportClient, const FVRAction VRAction, const EInputEvent Event, bool& bIsInputCaptured, bool& bOutWasHandled );
-
-	/** Called every frame to update hover state */
-	void OnVRHoverUpdate( FEditorViewportClient& ViewportClient, const int32 HandIndex, FVector& HoverImpactPoint, bool& bIsHoveringOverUI, bool& bOutWasHandled );
-
-	/** Called when the world is "dropped" with a hand */
-	void StopGrabbingWorld( struct FVirtualHand& VirtualHand );
-
-	/** Called when selection changes in the level */
-	void OnActorSelectionChanged( UObject* ChangedObject );
-
-	/** Checks to see if we're allowed to interact with the specified component */
-	bool IsInteractableComponent( const UActorComponent* Component ) const;
-
-	/** Called by the world interaction system when one of our components is dragged by the user.  If null is
-	    passed in then we'll treat it as dragging the whole object (rather than a specific axis/handle) */
-	void UpdateDragging( 
-		const float DeltaTime,
-		bool& bIsFirstDragUpdate, 
-		const EVREditorDraggingMode DraggingMode, 
-		const ETransformGizmoInteractionType InteractionType, 
-		const bool bWithTwoHands, 
-		class FEditorViewportClient& ViewportClient, 
-		const TOptional<FTransformGizmoHandlePlacement> OptionalHandlePlacement, 
-		const FVector& DragDelta, 
-		const FVector& OtherHandDragDelta, 
-		const FVector& DraggedTo, 
-		const FVector& OtherHandDraggedTo, 
-		const FVector& DragDeltaFromStart, 
-		const FVector& OtherHandDragDeltaFromStart, 
-		const FVector& LaserPointerStart, 
-		const FVector& LaserPointerDirection, 
-		const bool bIsLaserPointerValid, 
-		const FTransform& GizmoStartTransform, 
-		const FBox& GizmoStartLocalBounds, 
-		const USceneComponent* const DraggingTransformGizmoComponent,
-		FVector& GizmoSpaceFirstDragUpdateOffsetAlongAxis,
-		FVector& DragDeltaFromStartOffset,
-		bool& bIsDrivingVelocityOfSimulatedTransformables,
-		FVector& OutUnsnappedDraggedTo );
-
-	/** Given a drag delta from a starting point, contrains that delta based on a gizmo handle axis */
-	FVector ComputeConstrainedDragDeltaFromStart( const bool bIsFirstDragUpdate, const ETransformGizmoInteractionType InteractionType, const TOptional<FTransformGizmoHandlePlacement> OptionalHandlePlacement, const FVector& DragDeltaFromStart, const FVector& LaserPointerStart, const FVector& LaserPointerDirection, const bool bIsLaserPointerValid, const FTransform& GizmoStartTransform, FVector& GizmoSpaceFirstDragUpdateOffsetAlongAxis, FVector& DragDeltaFromStartOffset ) const;
-
-	/** Refreshes the transform gizmo to match the currently selected actor set */
-	void RefreshTransformGizmo( const bool bNewObjectsSelected, bool bAllHandlesVisible, class UActorComponent* SingleVisibleHandle, const TArray< UActorComponent* >& HoveringOverHandles );
-
-	/** Starts dragging selected actors around, changing selection if needed.  Called when clicking in the world, or when placing new objects */
-	void StartDraggingActors( FVRAction VRAction, UActorComponent* ClickedComponent, const FVector HitLocation, const bool bIsPlacingActors );
-
 	/** Starts dragging a material, allowing the user to drop it on an object in the scene to place it */
-	void StartDraggingMaterialOrTexture( FVRAction VRAction, const FVector HitLocation, UObject* MaterialOrTextureAsset );
+	void StartDraggingMaterialOrTexture( UViewportInteractor* Interactor, const FViewportActionKeyInput& Action, const FVector HitLocation, UObject* MaterialOrTextureAsset );
 
 	/** Tries to place whatever material or texture that's being dragged on the object under the hand's laser pointer */
-	void PlaceDraggedMaterialOrTexture( const int32 HandIndex );
-
-	/** Called to finish a drag action with the specified hand */
-	void StopDragging( const int32 HandIndex );
-
-	/** True if we're actively interacting with the specified actor */
-	bool IsTransformingActor( AActor* Actor ) const;
+	void PlaceDraggedMaterialOrTexture( UViewportInteractor* Interactor );
 
 	/** Called when FEditorDelegates::OnAssetDragStarted is broadcast */
 	void OnAssetDragStartedFromContentBrowser( const TArray<FAssetData>& DraggedAssets, class UActorFactory* FactoryToUse );
 
-	/** Figures out where to place an object when tracing it against the scene using a laser pointer */
-	bool FindPlacementPointUnderLaser( const int32 HandIndex, FVector& OutHitLocation );
-	
-	/** Start teleporting, does a ray trace with the hand passed and calculates the locations for lerp movement in Teleport */
-	void StartTeleport( const int32 HandIndex );
-
-	/** Actually teleport to a position */
-	void Teleport( const float DeltaTime );
-
-	/** Checks to see if smooth snapping is currently enabled */
-	bool IsSmoothSnappingEnabled() const;
-
-
 protected:
 
 	/** Owning object */
-	class FVREditorMode& Owner;
-
-
-	//
-	// Undo/redo
-	//
-
-	/** Manages saving undo for selected actors while we're dragging them around */
-	FTrackingTransaction TrackingTransaction;
+	FVREditorMode* Owner;
 
+	/** Sound for dropping materials and textures */
+	UPROPERTY()
+	class USoundCue* DropMaterialOrMaterialSound;
 
 	//
-	// Object movement/placement
+	// Dragging object from UI
 	//
 
-	/** All of the objects we're currently interacting with, such as selected actors */
-	TArray< TUniquePtr< class FVREditorTransformable > > Transformables;
-
-	/** True if our transformables are currently interpolating to their target location */
-	bool bIsInterpolatingTransformablesFromSnapshotTransform;
-
-	/** Time that we started interpolating at */
-	FTimespan TransformablesInterpolationStartTime;
-
-	/** Duration to interpolate transformables over */
-	float TransformablesInterpolationDuration;
+	/** The UI used to drag an asset into the level */
+	UPROPERTY()
+	class UWidgetComponent* FloatingUIAssetDraggedFrom;
 
 	/** The material or texture asset we're dragging to place on an object */
+	UPROPERTY()
 	UObject* PlacingMaterialOrTextureAsset;
-
-	//
-	// World movement
-	//
-
-	/** The world to meters scale before we changed it (last frame's world to meters scale) */
-	float LastWorldToMetersScale;
-
-
-	//
-	// Transform gizmo
-	//
-
-	/** Transform gizmo actor, for manipulating selected actor(s) */
-	class ABaseTransformGizmo* TransformGizmoActor;
-
-	/** Current gizmo bounds in local space.  Updated every frame.  This is not the same as the gizmo actor's bounds. */
-	FBox GizmoLocalBounds;
-
-	/** True if we've dragged objects with either hand since the last time we selected something */
-	bool bDraggedSinceLastSelection;
-
-	/** Gizmo start transform of the last drag we did with either hand.  Only valid if bDraggedSinceLastSelection is true */
-	FTransform LastDragGizmoStartTransform;
-
-	/** Starting angle when rotating an object with  ETransformGizmoInteractionType::RotateOnAngle */
-	TOptional<float> StartDragAngleOnRotation;
-
-	//
-	// Hover state
-	//
-
-	/** Set of objects that are currently hovered over */
-	TSet<FViewportHoverTarget> HoveredObjects;
-
-	//
-	// Teleport
-	//
-
-	/** If we are currently teleporting */
-	bool bIsTeleporting;
-
-	/** The current lerp of the teleport between the TeleportStartLocation and the TeleportGoalLocation */
-	float TeleportLerpAlpha;
-
-	/** Set on the current Roomspace location in the world in StartTeleport before doing the actual teleporting */
-	FVector TeleportStartLocation;
-
-	/** The calculated goal location in StartTeleport to move the Roomspace to */
-	FVector TeleportGoalLocation;
-
-	/** So we can teleport on a light press and make sure a second teleport by a fullpress doesn't happen */
-	bool bJustTeleported;
-
-	//
-	// Sound
-	//
-
-	/** Teleport sound */
-	class USoundCue* TeleportSound;
-
-	/** Sound for dropping materials and textures */
-	class USoundCue* DropMaterialOrMaterialSound;
-
-	/** The distance when starting dragging docks */
-	float DockSelectDistance;
-
-	/** The UI used to drag an asset into the level */
-	class UWidgetComponent* FloatingUIAssetDraggedFrom;
 };
 
diff --git a/Engine/Source/Editor/ViewportInteraction/Gizmo/VIBaseTransformGizmo.cpp b/Engine/Source/Editor/ViewportInteraction/Gizmo/VIBaseTransformGizmo.cpp
new file mode 100644
--- /dev/null
+++ b/Engine/Source/Editor/ViewportInteraction/Gizmo/VIBaseTransformGizmo.cpp
@@ -0,0 +1,275 @@
+// Copyright 1998-2015 Epic Games, Inc. All Rights Reserved.
+
+#include "ViewportInteractionModule.h"
+#include "VIBaseTransformGizmo.h"
+#include "UnitConversion.h"
+#include "VIGizmoHandle.h"
+#include "ViewportWorldInteraction.h"
+
+namespace VREd
+{
+	static FAutoConsoleVariable GizmoSelectionAnimationDuration( TEXT( "VREd.GizmoSelectionAnimationDuration" ), 0.15f, TEXT( "How long to animate the gizmo after objects are selected" ) );
+	static FAutoConsoleVariable GizmoSelectionAnimationCurvePower( TEXT( "VREd.GizmoSelectionAnimationCurvePower" ), 2.0f, TEXT( "Controls the animation curve for the gizmo after objects are selected" ) );
+	static FAutoConsoleVariable GizmoShowMeasurementText( TEXT( "VREd.GizmoShowMeasurementText" ), 0, TEXT( "When enabled, gizmo measurements will always be visible.  Otherwise, only when hovering over a scale/stretch gizmo handle" ) );
+}
+
+ABaseTransformGizmo::ABaseTransformGizmo( ) :
+	WorldInteraction( nullptr )
+{
+	// Create root default scene component
+	{
+		SceneComponent = CreateDefaultSubobject<USceneComponent>( TEXT( "SceneComponent" ) );
+		check( SceneComponent != nullptr );
+
+		RootComponent = SceneComponent;
+	}
+
+	GizmoMaterial = nullptr;
+	{
+		static ConstructorHelpers::FObjectFinder<UMaterial> ObjectFinder( TEXT( "/Engine/VREditor/TransformGizmo/TransformGizmoMaterial" ) );
+		GizmoMaterial = ObjectFinder.Object;
+		check( GizmoMaterial != nullptr );
+	}
+
+	TranslucentGizmoMaterial = nullptr;
+	{
+		static ConstructorHelpers::FObjectFinder<UMaterial> ObjectFinder( TEXT( "/Engine/VREditor/TransformGizmo/TranslucentTransformGizmoMaterial" ) );
+		TranslucentGizmoMaterial = ObjectFinder.Object;
+		check( TranslucentGizmoMaterial != nullptr );
+	}
+}
+
+ABaseTransformGizmo::~ABaseTransformGizmo()
+{
+	WorldInteraction = nullptr;
+}
+
+void ABaseTransformGizmo::OnNewObjectsSelected()
+{
+	SelectedAtTime = FTimespan::FromSeconds( FApp::GetCurrentTime() );
+}
+
+ETransformGizmoInteractionType ABaseTransformGizmo::GetInteractionType( UActorComponent* DraggedComponent, TOptional<FTransformGizmoHandlePlacement>& OutHandlePlacement )
+{
+	OutHandlePlacement.Reset();
+	if ( DraggedComponent != nullptr )
+	{
+		UStaticMeshComponent* DraggedMesh = Cast<UStaticMeshComponent>( DraggedComponent );
+		if ( DraggedMesh != nullptr )
+		{
+			ETransformGizmoInteractionType ResultInteractionType;
+
+			for ( UGizmoHandleGroup* HandleGroup : AllHandleGroups )
+			{
+				if ( HandleGroup != nullptr )
+				{
+					int32 HandIndex = HandleGroup->GetDraggedHandleIndex( DraggedMesh );
+					if ( HandIndex != INDEX_NONE )
+					{
+						HandleGroup->GetHandleIndexInteractionType( HandIndex, ResultInteractionType, OutHandlePlacement );
+						return ResultInteractionType;
+					}
+				}
+			}
+		}
+	}
+
+	OutHandlePlacement.Reset();
+	return ETransformGizmoInteractionType::Translate;
+}
+
+float ABaseTransformGizmo::GetAnimationAlpha()
+{
+	// Update animation
+	float AnimationAlpha = 0.0f;
+	{
+		const FTimespan CurrentTime = FTimespan::FromSeconds( FApp::GetCurrentTime() );
+		const float TimeSinceSelectionChange = (CurrentTime - SelectedAtTime).GetTotalSeconds();
+		const float AnimLength = VREd::GizmoSelectionAnimationDuration->GetFloat();
+		if ( TimeSinceSelectionChange < AnimLength )
+		{
+			AnimationAlpha = FMath::Max( 0.0f, TimeSinceSelectionChange / AnimLength );
+		}
+		else
+		{
+			AnimationAlpha = 1.0f;
+		}
+
+		// Apply a bit of a curve to the animation
+		AnimationAlpha = FMath::Pow( AnimationAlpha, VREd::GizmoSelectionAnimationCurvePower->GetFloat() );
+	}
+
+	return AnimationAlpha;
+}
+
+void ABaseTransformGizmo::SetOwnerWorldInteraction( class UViewportWorldInteraction* InWorldInteraction )
+{
+	WorldInteraction = InWorldInteraction;
+}
+
+class UViewportWorldInteraction* ABaseTransformGizmo::GetOwnerWorldInteraction() const
+{
+	return WorldInteraction;
+}
+
+void ABaseTransformGizmo::GetBoundingBoxEdge( const FBox& Box, const int32 AxisIndex, const int32 EdgeIndex, FVector& OutVertex0, FVector& OutVertex1 )
+{
+	check( AxisIndex >= 0 && AxisIndex < 3 );
+	check( EdgeIndex >= 0 && EdgeIndex < 4 );
+
+	const FVector BackBottomLeft( Box.Min.X, Box.Min.Y, Box.Min.Z );
+	const FVector BackBottomRight( Box.Min.X, Box.Max.Y, Box.Min.Z );
+	const FVector BackTopLeft( Box.Min.X, Box.Min.Y, Box.Max.Z );
+	const FVector BackTopRight( Box.Min.X, Box.Max.Y, Box.Max.Z );
+
+	const FVector FrontBottomLeft( Box.Max.X, Box.Min.Y, Box.Min.Z );
+	const FVector FrontBottomRight( Box.Max.X, Box.Max.Y, Box.Min.Z );
+	const FVector FrontTopLeft( Box.Max.X, Box.Min.Y, Box.Max.Z );
+	const FVector FrontTopRight( Box.Max.X, Box.Max.Y, Box.Max.Z );
+
+	switch (AxisIndex)
+	{
+		case 0:	// Front/back
+		{
+			switch ( EdgeIndex )
+			{
+			case 0:	// Bottom left
+				OutVertex0 = BackBottomLeft;
+				OutVertex1 = FrontBottomLeft;
+				break;
+
+			case 1:	// Top left
+				OutVertex0 = BackTopLeft;
+				OutVertex1 = FrontTopLeft;
+				break;
+
+			case 2:	// Top right
+				OutVertex0 = BackTopRight;
+				OutVertex1 = FrontTopRight;
+				break;
+
+			case 3:	// Bottom right
+				OutVertex0 = BackBottomRight;
+				OutVertex1 = FrontBottomRight;
+				break;
+			}
+			break;
+		}
+
+		case 1:
+		{
+			switch ( EdgeIndex )
+			{
+			case 0:	// Back bottom
+				OutVertex0 = BackBottomLeft;
+				OutVertex1 = BackBottomRight;
+				break;
+
+			case 1:	// Back top
+				OutVertex0 = BackTopLeft;
+				OutVertex1 = BackTopRight;
+				break;
+
+			case 2:	// Front top
+				OutVertex0 = FrontTopLeft;
+				OutVertex1 = FrontTopRight;
+				break;
+
+			case 3:	// Front bottom
+				OutVertex0 = FrontBottomLeft;
+				OutVertex1 = FrontBottomRight;
+				break;
+			}
+			break;
+		}
+
+		case 2:
+		{
+			switch ( EdgeIndex )
+			{
+			case 0:	// Back left
+				OutVertex0 = BackBottomLeft;
+				OutVertex1 = BackTopLeft;
+				break;
+
+			case 1:	// Back right
+				OutVertex0 = BackBottomRight;
+				OutVertex1 = BackTopRight;
+				break;
+
+			case 2:	// Front right
+				OutVertex0 = FrontBottomRight;
+				OutVertex1 = FrontTopRight;
+				break;
+
+			case 3:	// Front left
+				OutVertex0 = FrontBottomLeft;
+				OutVertex1 = FrontTopLeft;
+				break;
+			}
+			break;
+		}
+	}
+}
+
+
+
+
+void ABaseTransformGizmo::UpdateHandleVisibility( const EGizmoHandleTypes GizmoType, const ECoordSystem GizmoCoordinateSpace, const bool bAllHandlesVisible, UActorComponent* DraggingHandle )
+{
+	for ( UGizmoHandleGroup* HandleGroup : AllHandleGroups )
+	{
+		if ( HandleGroup != nullptr )
+		{
+			const bool bIsTypeSupported = ( GizmoType == EGizmoHandleTypes::All && HandleGroup->GetShowOnUniversalGizmo() ) || HandleGroup->GetHandleType() == GizmoType;
+			const bool bSupportsCurrentCoordinateSpace = HandleGroup->SupportsWorldCoordinateSpace() || GizmoCoordinateSpace != COORD_World;
+
+			for ( FGizmoHandle& Handle : HandleGroup->GetHandles() )
+			{
+				if( Handle.HandleMesh != nullptr )
+				{
+					const bool bShowIt = ( bIsTypeSupported && bSupportsCurrentCoordinateSpace && bAllHandlesVisible ) || ( DraggingHandle != nullptr && DraggingHandle == Handle.HandleMesh );
+
+					Handle.HandleMesh->SetVisibility( bShowIt );
+
+					// Never allow ray queries to impact hidden handles
+					Handle.HandleMesh->SetCollisionEnabled( bShowIt ? ECollisionEnabled::QueryOnly : ECollisionEnabled::NoCollision );
+				}
+			}
+		}
+	}
+}
+
+bool ABaseTransformGizmo::GetShowMeasurementText() const
+{
+	return VREd::GizmoShowMeasurementText->GetInt() != 0;
+}
+
+void FTransformGizmoHandlePlacement::GetCenterHandleCountAndFacingAxisIndex( int32& OutCenterHandleCount, int32& OutFacingAxisIndex, int32& OutCenterAxisIndex ) const
+{
+	OutCenterHandleCount = 0;
+	OutFacingAxisIndex = INDEX_NONE;
+	OutCenterAxisIndex = INDEX_NONE;
+	for ( int32 AxisIndex = 0; AxisIndex < 3; ++AxisIndex )
+	{
+		if ( Axes[AxisIndex] == ETransformGizmoHandleDirection::Center )
+		{
+			++OutCenterHandleCount;
+			OutCenterAxisIndex = AxisIndex;
+		}
+		else
+		{
+			OutFacingAxisIndex = AxisIndex;
+		}
+	}
+
+	if ( OutCenterHandleCount < 2 )
+	{
+		OutFacingAxisIndex = INDEX_NONE;
+	}
+
+	if ( OutCenterHandleCount != 1 )
+	{
+		OutCenterAxisIndex = INDEX_NONE;
+	}
+}
\ No newline at end of file
diff --git a/Engine/Source/Editor/ViewportInteraction/Gizmo/VIPivotTransformGizmo.cpp b/Engine/Source/Editor/ViewportInteraction/Gizmo/VIPivotTransformGizmo.cpp
new file mode 100644
--- /dev/null
+++ b/Engine/Source/Editor/ViewportInteraction/Gizmo/VIPivotTransformGizmo.cpp
@@ -0,0 +1,250 @@
+// Copyright 1998-2015 Epic Games, Inc. All Rights Reserved.
+
+#include "ViewportInteractionModule.h"
+#include "VIPivotTransformGizmo.h"
+#include "VIStretchGizmoHandle.h"
+#include "VIUniformScaleGizmoHandle.h"
+
+namespace VREd //@todo VREditor: Duplicates of TransformGizmo
+{
+	// @todo vreditor tweak: Tweak out console variables
+	static FAutoConsoleVariable PivotGizmoDistanceScaleFactor( TEXT( "VREd.PivotGizmoDistanceScaleFactor" ), 0.003f, TEXT( "How much the gizmo handles should increase in size with distance from the camera, to make it easier to select"));
+	static FAutoConsoleVariable PivotGizmoTranslationPivotOffsetX( TEXT("VREd.PivotGizmoTranslationPivotOffsetX" ), 20.0f, TEXT( "How much the translation cylinder is offsetted from the pivot" ) );
+	static FAutoConsoleVariable PivotGizmoScalePivotOffsetX( TEXT( "VREd.PivotGizmoScalePivotOffsetX" ), 120.0f, TEXT( "How much the non-uniform scale is offsetted from the pivot" ) );
+	static FAutoConsoleVariable PivotGizmoPlaneTranslationPivotOffsetYZ(TEXT("VREd.PivotGizmoPlaneTranslationPivotOffsetYZ" ), 40.0f, TEXT( "How much the plane translation is offsetted from the pivot" ) );
+	static FAutoConsoleVariable PivotGizmoTranslationScaleMultiply( TEXT( "VREd.PivotGizmoTranslationScaleMultiply" ), 3.0f, TEXT( "Multiplies translation handles scale" ) );
+	static FAutoConsoleVariable PivotGizmoTranslationHoverScaleMultiply( TEXT( "VREd.PivotGizmoTranslationHoverScaleMultiply" ), 1.2f, TEXT( "Multiplies translation handles hover scale" ) );
+}
+
+APivotTransformGizmo::APivotTransformGizmo() :
+	Super()
+{
+	UniformScaleGizmoHandleGroup = CreateDefaultSubobject<UUniformScaleGizmoHandleGroup>( TEXT( "UniformScaleHandles" ), true );
+	UniformScaleGizmoHandleGroup->SetTranslucentGizmoMaterial( TranslucentGizmoMaterial );
+	UniformScaleGizmoHandleGroup->SetGizmoMaterial( GizmoMaterial );
+	UniformScaleGizmoHandleGroup->SetupAttachment( SceneComponent );
+	AllHandleGroups.Add( UniformScaleGizmoHandleGroup );
+
+	TranslationGizmoHandleGroup = CreateDefaultSubobject<UPivotTranslationGizmoHandleGroup>(TEXT("TranslationHandles"), true);
+	TranslationGizmoHandleGroup->SetTranslucentGizmoMaterial(TranslucentGizmoMaterial);
+	TranslationGizmoHandleGroup->SetGizmoMaterial(GizmoMaterial);
+	TranslationGizmoHandleGroup->SetupAttachment( SceneComponent );
+	AllHandleGroups.Add(TranslationGizmoHandleGroup);
+
+	ScaleGizmoHandleGroup = CreateDefaultSubobject<UPivotScaleGizmoHandleGroup>( TEXT( "ScaleHandles" ), true );
+	ScaleGizmoHandleGroup->SetTranslucentGizmoMaterial( TranslucentGizmoMaterial );
+	ScaleGizmoHandleGroup->SetGizmoMaterial( GizmoMaterial );
+	ScaleGizmoHandleGroup->SetupAttachment( SceneComponent );
+	AllHandleGroups.Add( ScaleGizmoHandleGroup );
+
+	PlaneTranslationGizmoHandleGroup = CreateDefaultSubobject<UPivotPlaneTranslationGizmoHandleGroup>( TEXT( "PlaneTranslationHandles" ), true );
+	PlaneTranslationGizmoHandleGroup->SetTranslucentGizmoMaterial( TranslucentGizmoMaterial );
+	PlaneTranslationGizmoHandleGroup->SetGizmoMaterial( GizmoMaterial );
+	PlaneTranslationGizmoHandleGroup->SetupAttachment( SceneComponent );
+	AllHandleGroups.Add( PlaneTranslationGizmoHandleGroup );
+
+	RotationGizmoHandleGroup = CreateDefaultSubobject<UPivotRotationGizmoHandleGroup>( TEXT( "RotationHandles" ), true );
+	RotationGizmoHandleGroup->SetTranslucentGizmoMaterial( TranslucentGizmoMaterial );
+	RotationGizmoHandleGroup->SetGizmoMaterial( GizmoMaterial );
+	RotationGizmoHandleGroup->SetupAttachment( SceneComponent );
+	AllHandleGroups.Add( RotationGizmoHandleGroup );
+
+	StretchGizmoHandleGroup = CreateDefaultSubobject<UStretchGizmoHandleGroup>( TEXT( "StretchHandles" ), true );
+	StretchGizmoHandleGroup->SetTranslucentGizmoMaterial( TranslucentGizmoMaterial );
+	StretchGizmoHandleGroup->SetGizmoMaterial( GizmoMaterial );
+	StretchGizmoHandleGroup->SetShowOnUniversalGizmo( false );
+	StretchGizmoHandleGroup->SetupAttachment( SceneComponent );
+	AllHandleGroups.Add( StretchGizmoHandleGroup );
+
+	// There may already be some objects selected as we switch into VR mode, so we'll pretend that just happened so
+	// that we can make sure all transitions complete properly
+	OnNewObjectsSelected();
+}
+
+void APivotTransformGizmo::UpdateGizmo( const EGizmoHandleTypes GizmoType, const ECoordSystem GizmoCoordinateSpace, const FTransform& LocalToWorld, const FBox& LocalBounds, const FVector ViewLocation, bool bAllHandlesVisible, 
+	UActorComponent* DraggingHandle, const TArray< UActorComponent* >& HoveringOverHandles, const float GizmoHoverScale, const float GizmoHoverAnimationDuration )
+{
+	const float WorldScaleFactor = GetWorld()->GetWorldSettings()->WorldToMeters / 100.0f;
+
+	// Position the gizmo at the location of the first selected actor
+	const bool bSweep = false;
+	this->SetActorTransform( LocalToWorld, bSweep );
+
+	// Increase scale with distance, to make gizmo handles easier to click on
+	// @todo vreditor: Should probably be a curve, not linear
+	// @todo vreditor: Should take FOV into account (especially in non-stereo/HMD mode)
+	const float WorldSpaceDistanceToToPivot = FMath::Sqrt( FVector::DistSquared( GetActorLocation(), ViewLocation ) );
+	const float GizmoScale( ( ( WorldSpaceDistanceToToPivot / WorldScaleFactor ) * VREd::PivotGizmoDistanceScaleFactor->GetFloat() ) * WorldScaleFactor );
+
+	// Update animation
+	float AnimationAlpha = GetAnimationAlpha();
+
+	// Update all the handles
+	bool bIsHoveringOrDraggingScaleGizmo = false;
+	for ( UGizmoHandleGroup* HandleGroup : AllHandleGroups )
+	{
+		if ( HandleGroup != nullptr && 
+		   ( DraggingHandle == nullptr || HandleGroup == StretchGizmoHandleGroup ) )
+		{
+			bool bIsHoveringOrDraggingThisHandleGroup = false;
+			HandleGroup->UpdateGizmoHandleGroup( LocalToWorld, LocalBounds, ViewLocation, bAllHandlesVisible, DraggingHandle,
+				HoveringOverHandles, AnimationAlpha, GizmoScale, GizmoHoverScale, GizmoHoverAnimationDuration, /* Out */ bIsHoveringOrDraggingThisHandleGroup );
+		}
+	}
+
+	UpdateHandleVisibility( GizmoType, GizmoCoordinateSpace, bAllHandlesVisible, DraggingHandle );
+}
+
+/************************************************************************/
+/* Translation                                                          */
+/************************************************************************/
+UPivotTranslationGizmoHandleGroup::UPivotTranslationGizmoHandleGroup() :
+	Super()
+{
+	UStaticMesh* TranslationHandleMesh = nullptr;
+	{
+		static ConstructorHelpers::FObjectFinder<UStaticMesh> ObjectFinder( TEXT( "/Engine/VREditor/TransformGizmo/TranslateHandleLong" ) );
+		TranslationHandleMesh = ObjectFinder.Object;
+		check( TranslationHandleMesh != nullptr );
+	}
+
+	CreateHandles( TranslationHandleMesh, FString( "PivotTranslationHandle" ) );
+}
+
+
+void UPivotTranslationGizmoHandleGroup::UpdateGizmoHandleGroup( const FTransform& LocalToWorld, const FBox& LocalBounds, const FVector ViewLocation, bool bAllHandlesVisible, class UActorComponent* DraggingHandle, 
+	const TArray< UActorComponent* >& HoveringOverHandles, float AnimationAlpha, float GizmoScale, const float GizmoHoverScale, const float GizmoHoverAnimationDuration, bool& bOutIsHoveringOrDraggingThisHandleGroup )
+{
+	// Call parent implementation (updates hover animation)
+	Super::UpdateGizmoHandleGroup( LocalToWorld, LocalBounds, ViewLocation, bAllHandlesVisible, DraggingHandle, HoveringOverHandles, AnimationAlpha,
+		GizmoScale, GizmoHoverScale, GizmoHoverAnimationDuration, bOutIsHoveringOrDraggingThisHandleGroup );
+	
+	const float MultipliedGizmoScale = GizmoScale * VREd::PivotGizmoTranslationScaleMultiply->GetFloat();
+	const float MultipliedGizmoHoverScale = GizmoHoverScale * VREd::PivotGizmoTranslationHoverScaleMultiply->GetFloat();
+
+	UpdateHandlesRelativeTransformOnAxis( FTransform( FVector( VREd::PivotGizmoTranslationPivotOffsetX->GetFloat(), 0, 0 ) ), AnimationAlpha, MultipliedGizmoScale, MultipliedGizmoHoverScale, ViewLocation, DraggingHandle, HoveringOverHandles );
+}
+
+ETransformGizmoInteractionType UPivotTranslationGizmoHandleGroup::GetInteractionType() const
+{
+	return ETransformGizmoInteractionType::Translate;
+}
+
+EGizmoHandleTypes UPivotTranslationGizmoHandleGroup::GetHandleType() const
+{
+	return EGizmoHandleTypes::Translate;
+}
+
+/************************************************************************/
+/* Scale	                                                            */
+/************************************************************************/
+UPivotScaleGizmoHandleGroup::UPivotScaleGizmoHandleGroup() :
+	Super()
+{
+	UStaticMesh* TranslationHandleMesh = nullptr;
+	{
+		static ConstructorHelpers::FObjectFinder<UStaticMesh> ObjectFinder( TEXT( "/Engine/VREditor/TransformGizmo/UniformScaleHandle" ) );
+		TranslationHandleMesh = ObjectFinder.Object;
+		check( TranslationHandleMesh != nullptr );
+	}
+
+	CreateHandles( TranslationHandleMesh, FString( "PivotScaleHandle" ) );	
+}
+
+void UPivotScaleGizmoHandleGroup::UpdateGizmoHandleGroup( const FTransform& LocalToWorld, const FBox& LocalBounds, const FVector ViewLocation, bool bAllHandlesVisible, class UActorComponent* DraggingHandle, 
+	const TArray< UActorComponent* >& HoveringOverHandles, float AnimationAlpha, float GizmoScale, const float GizmoHoverScale, const float GizmoHoverAnimationDuration, bool& bOutIsHoveringOrDraggingThisHandleGroup )
+{
+	// Call parent implementation (updates hover animation)
+	Super::UpdateGizmoHandleGroup( LocalToWorld, LocalBounds, ViewLocation, bAllHandlesVisible, DraggingHandle, HoveringOverHandles, AnimationAlpha, 
+		GizmoScale, GizmoHoverScale, GizmoHoverAnimationDuration, bOutIsHoveringOrDraggingThisHandleGroup );
+
+	UpdateHandlesRelativeTransformOnAxis( FTransform( FVector( VREd::PivotGizmoScalePivotOffsetX->GetFloat(), 0, 0 ) ), AnimationAlpha, GizmoScale, GizmoHoverScale, ViewLocation, DraggingHandle, HoveringOverHandles );
+}
+
+ETransformGizmoInteractionType UPivotScaleGizmoHandleGroup::GetInteractionType() const
+{
+	return ETransformGizmoInteractionType::Scale;
+}
+
+EGizmoHandleTypes UPivotScaleGizmoHandleGroup::GetHandleType() const
+{
+	return EGizmoHandleTypes::Scale;
+}
+
+bool UPivotScaleGizmoHandleGroup::SupportsWorldCoordinateSpace() const
+{
+	return false;
+}
+
+/************************************************************************/
+/* Plane Translation	                                                */
+/************************************************************************/
+UPivotPlaneTranslationGizmoHandleGroup::UPivotPlaneTranslationGizmoHandleGroup() :
+	Super()
+{
+	UStaticMesh* TranslationHandleMesh = nullptr;
+	{
+		static ConstructorHelpers::FObjectFinder<UStaticMesh> ObjectFinder( TEXT( "/Engine/VREditor/TransformGizmo/PlaneTranslationHandle" ) );
+		TranslationHandleMesh = ObjectFinder.Object;
+		check( TranslationHandleMesh != nullptr );
+	}
+
+	CreateHandles( TranslationHandleMesh, FString( "PlaneTranslationHandle" ) );
+}
+
+void UPivotPlaneTranslationGizmoHandleGroup::UpdateGizmoHandleGroup( const FTransform& LocalToWorld, const FBox& LocalBounds, const FVector ViewLocation, bool bAllHandlesVisible, class UActorComponent* DraggingHandle,
+	const TArray< UActorComponent* >& HoveringOverHandles, float AnimationAlpha, float GizmoScale, const float GizmoHoverScale, const float GizmoHoverAnimationDuration, bool& bOutIsHoveringOrDraggingThisHandleGroup )
+{
+	// Call parent implementation (updates hover animation)
+	Super::UpdateGizmoHandleGroup( LocalToWorld, LocalBounds, ViewLocation, bAllHandlesVisible, DraggingHandle, HoveringOverHandles, AnimationAlpha,
+		GizmoScale, GizmoHoverScale, GizmoHoverAnimationDuration, bOutIsHoveringOrDraggingThisHandleGroup );
+
+	UpdateHandlesRelativeTransformOnAxis( FTransform( FVector( 0, VREd::PivotGizmoPlaneTranslationPivotOffsetYZ->GetFloat(), VREd::PivotGizmoPlaneTranslationPivotOffsetYZ->GetFloat() ) ), 
+		AnimationAlpha, GizmoScale, GizmoHoverScale, ViewLocation, DraggingHandle, HoveringOverHandles );
+}
+
+ETransformGizmoInteractionType UPivotPlaneTranslationGizmoHandleGroup::GetInteractionType() const
+{
+	return ETransformGizmoInteractionType::TranslateOnPlane;
+}
+
+EGizmoHandleTypes UPivotPlaneTranslationGizmoHandleGroup::GetHandleType() const
+{
+	return EGizmoHandleTypes::Translate;
+}
+
+/************************************************************************/
+/* Rotation																*/
+/************************************************************************/
+UPivotRotationGizmoHandleGroup::UPivotRotationGizmoHandleGroup() :
+	Super()
+{
+	UStaticMesh* TranslationHandleMesh = nullptr;
+	{
+		static ConstructorHelpers::FObjectFinder<UStaticMesh> ObjectFinder(TEXT("/Engine/VREditor/TransformGizmo/RotationHandleFull"));
+		TranslationHandleMesh = ObjectFinder.Object;
+		check(TranslationHandleMesh != nullptr);
+	}
+
+	CreateHandles(TranslationHandleMesh, FString("RotationHandle"));
+}
+
+void UPivotRotationGizmoHandleGroup::UpdateGizmoHandleGroup( const FTransform& LocalToWorld, const FBox& LocalBounds, const FVector ViewLocation, bool bAllHandlesVisible, class UActorComponent* DraggingHandle,
+	const TArray< UActorComponent* >& HoveringOverHandles, float AnimationAlpha, float GizmoScale, const float GizmoHoverScale, const float GizmoHoverAnimationDuration, bool& bOutIsHoveringOrDraggingThisHandleGroup )
+{
+	// Call parent implementation (updates hover animation)
+	Super::UpdateGizmoHandleGroup(LocalToWorld, LocalBounds, ViewLocation, bAllHandlesVisible, DraggingHandle, HoveringOverHandles, AnimationAlpha,
+		GizmoScale, GizmoHoverScale, GizmoHoverAnimationDuration, bOutIsHoveringOrDraggingThisHandleGroup );
+	
+	UpdateHandlesRelativeTransformOnAxis( FTransform( FVector( 0, 0, 0 ) ), AnimationAlpha, GizmoScale, GizmoScale, ViewLocation, DraggingHandle, HoveringOverHandles );
+}
+
+ETransformGizmoInteractionType UPivotRotationGizmoHandleGroup::GetInteractionType() const
+{
+	return ETransformGizmoInteractionType::RotateOnAngle;
+}
+
+EGizmoHandleTypes UPivotRotationGizmoHandleGroup::GetHandleType() const
+{
+	return EGizmoHandleTypes::Rotate;
+}
\ No newline at end of file
diff --git a/Engine/Source/Editor/ViewportInteraction/Gizmo/VIPivotTransformGizmo.h b/Engine/Source/Editor/ViewportInteraction/Gizmo/VIPivotTransformGizmo.h
new file mode 100644
--- /dev/null
+++ b/Engine/Source/Editor/ViewportInteraction/Gizmo/VIPivotTransformGizmo.h
@@ -0,0 +1,152 @@
+// Copyright 1998-2015 Epic Games, Inc. All Rights Reserved.
+
+#pragma once
+
+#include "VIBaseTransformGizmo.h"
+#include "VIGizmoHandle.h"
+#include "VIPivotTransformGizmo.generated.h"
+
+/**
+ * A transform gizmo on the pivot that allows you to interact with selected objects by moving, scaling and rotating.
+ */
+UCLASS()
+class APivotTransformGizmo : public ABaseTransformGizmo
+{
+	GENERATED_BODY()
+
+public:
+
+	/** Default constructor that sets up CDO properties */
+	APivotTransformGizmo();
+
+	/** Called by the world interaction system after we've been spawned into the world, to allow
+	    us to create components and set everything up nicely for the selected objects that we'll be
+		used to manipulate */
+	virtual void UpdateGizmo( const EGizmoHandleTypes GizmoType, const ECoordSystem GizmoCoordinateSpace, const FTransform& LocalToWorld, const FBox& LocalBounds, const FVector ViewLocation, 
+		bool bAllHandlesVisible, class UActorComponent* DraggingHandle, const TArray< UActorComponent* >& HoveringOverHandles, const float GizmoHoverScale, const float GizmoHoverAnimationDuration ) override;
+
+private:
+	/** Uniform scale handle group component */
+	UPROPERTY()
+	class UUniformScaleGizmoHandleGroup* UniformScaleGizmoHandleGroup;
+
+	/** Translation handle group component */
+	UPROPERTY()
+	class UPivotTranslationGizmoHandleGroup* TranslationGizmoHandleGroup;
+	
+	/** Scale handle group component */
+	UPROPERTY()
+	class UPivotScaleGizmoHandleGroup* ScaleGizmoHandleGroup;
+	
+	/** Plane translation handle group component */
+	UPROPERTY()
+	class UPivotPlaneTranslationGizmoHandleGroup* PlaneTranslationGizmoHandleGroup;
+
+	/** Rotation handle group component */
+	UPROPERTY()
+	class UPivotRotationGizmoHandleGroup* RotationGizmoHandleGroup;
+
+	/** Stretch handle group component */
+	UPROPERTY()
+	class UStretchGizmoHandleGroup* StretchGizmoHandleGroup;
+};
+
+/**
+ * Axis Gizmo handle for translating
+ */
+UCLASS()
+class UPivotTranslationGizmoHandleGroup : public UAxisGizmoHandleGroup
+{
+	GENERATED_BODY()
+
+public:
+
+	/** Default constructor that sets up CDO properties */
+	UPivotTranslationGizmoHandleGroup();
+
+	/** Updates the gizmo handles */
+	virtual void UpdateGizmoHandleGroup( const FTransform& LocalToWorld, const FBox& LocalBounds, const FVector ViewLocation, bool bAllHandlesVisible, class UActorComponent* DraggingHandle, const TArray< UActorComponent* >& HoveringOverHandles, 
+		float AnimationAlpha, float GizmoScale, const float GizmoHoverScale, const float GizmoHoverAnimationDuration, bool& bOutIsHoveringOrDraggingThisHandleGroup ) override;
+
+	/** Gets the InteractionType for this Gizmo handle */
+	virtual ETransformGizmoInteractionType GetInteractionType() const override;
+
+	/** Gets the GizmoType for this Gizmo handle */
+	virtual EGizmoHandleTypes GetHandleType() const override;
+};
+
+/**
+ * Axis Gizmo handle for scaling
+ */
+UCLASS()
+class UPivotScaleGizmoHandleGroup : public UAxisGizmoHandleGroup
+{
+	GENERATED_BODY()
+
+public:
+
+	/** Default constructor that sets up CDO properties */
+	UPivotScaleGizmoHandleGroup();
+
+	/** Updates the gizmo handles */
+	virtual void UpdateGizmoHandleGroup( const FTransform& LocalToWorld, const FBox& LocalBounds,  const FVector ViewLocation, bool bAllHandlesVisible, class UActorComponent* DraggingHandle, const TArray< UActorComponent* >& HoveringOverHandles, 
+		float AnimationAlpha, float GizmoScale, const float GizmoHoverScale, const float GizmoHoverAnimationDuration, bool& bOutIsHoveringOrDraggingThisHandleGroup ) override;
+
+	/** Gets the InteractionType for this Gizmo handle */
+	virtual ETransformGizmoInteractionType GetInteractionType() const override;
+
+	/** Gets the GizmoType for this Gizmo handle */
+	virtual EGizmoHandleTypes GetHandleType() const override;
+
+	/** Returns true if this type of handle is allowed with world space gizmos */
+	virtual bool SupportsWorldCoordinateSpace() const override;
+};
+
+/**
+ * Axis Gizmo handle for plane translation
+ */
+UCLASS()
+class UPivotPlaneTranslationGizmoHandleGroup : public UAxisGizmoHandleGroup
+{
+	GENERATED_BODY()
+
+public:
+
+	/** Default constructor that sets up CDO properties */
+	UPivotPlaneTranslationGizmoHandleGroup();
+
+	/** Updates the gizmo handles */
+	virtual void UpdateGizmoHandleGroup( const FTransform& LocalToWorld, const FBox& LocalBounds, const FVector ViewLocation, bool bAllHandlesVisible, class UActorComponent* DraggingHandle, const TArray< UActorComponent* >& HoveringOverHandles, 
+		float AnimationAlpha, float GizmoScale, const float GizmoHoverScale, const float GizmoHoverAnimationDuration, bool& bOutIsHoveringOrDraggingThisHandleGroup ) override;
+
+	/** Gets the InteractionType for this Gizmo handle */
+	virtual ETransformGizmoInteractionType GetInteractionType() const override;
+
+	/** Gets the GizmoType for this Gizmo handle */
+	virtual EGizmoHandleTypes GetHandleType() const override;
+};
+
+
+/**
+ * Axis Gizmo handle for rotation
+ */
+UCLASS()
+class UPivotRotationGizmoHandleGroup : public UAxisGizmoHandleGroup
+{
+	GENERATED_BODY()
+
+public:
+
+	/** Default constructor that sets up CDO properties */
+	UPivotRotationGizmoHandleGroup();
+
+	/** Updates the gizmo handles */
+	virtual void UpdateGizmoHandleGroup(const FTransform& LocalToWorld, const FBox& LocalBounds, const FVector ViewLocation, bool bAllHandlesVisible, class UActorComponent* DraggingHandle, const TArray< UActorComponent* >& HoveringOverHandles, 
+		float AnimationAlpha, float GizmoScale, const float GizmoHoverScale, const float GizmoHoverAnimationDuration, bool& bOutIsHoveringOrDraggingThisHandleGroup ) override;
+
+	/** Gets the InteractionType for this Gizmo handle */
+	virtual ETransformGizmoInteractionType GetInteractionType() const override;
+
+	/** Gets the GizmoType for this Gizmo handle */
+	virtual EGizmoHandleTypes GetHandleType() const override;
+};
\ No newline at end of file
diff --git a/Engine/Source/Editor/ViewportInteraction/Public/VIBaseTransformGizmo.h b/Engine/Source/Editor/ViewportInteraction/Public/VIBaseTransformGizmo.h
new file mode 100644
--- /dev/null
+++ b/Engine/Source/Editor/ViewportInteraction/Public/VIBaseTransformGizmo.h
@@ -0,0 +1,104 @@
+// Copyright 1998-2015 Epic Games, Inc. All Rights Reserved.
+
+#pragma once
+
+#include "ViewportInteractionTypes.h"
+#include "GameFramework/Actor.h"
+#include "VIBaseTransformGizmo.generated.h"
+
+UENUM()
+enum class EGizmoHandleTypes : uint8
+{
+	All = 0,
+	Translate = 1,
+	Rotate = 2,
+	Scale = 3
+};
+
+/**
+* Displays measurements along the bounds of selected objects
+*/
+USTRUCT()
+struct FTransformGizmoMeasurement
+{
+	GENERATED_BODY()
+
+		/** The text that displays the actual measurement and units */
+		UPROPERTY()
+	class UTextRenderComponent* MeasurementText;
+};
+
+
+/**
+ * Base class for transform gizmo
+ */
+UCLASS( Abstract )
+class VIEWPORTINTERACTION_API ABaseTransformGizmo : public AActor
+{
+	GENERATED_BODY()
+
+public:
+	
+	/** Default constructor that sets up CDO properties */
+	ABaseTransformGizmo();
+
+	/** Deconstructor */
+	virtual ~ABaseTransformGizmo();
+
+	/** Call this when new objects become selected.  This triggers an animation transition. */
+	void OnNewObjectsSelected();
+
+	/** Called by the world interaction system when one of our components is dragged by the user to find out
+	    what type of interaction to do.  If null is passed in then we'll treat it as dragging the whole object 
+		(rather than a specific axis/handle) */
+	ETransformGizmoInteractionType GetInteractionType( UActorComponent* DraggedComponent, TOptional<FTransformGizmoHandlePlacement>& OutHandlePlacement );
+
+	/** Updates the animation with the current time and selected time */
+	float GetAnimationAlpha();
+
+	/** Sets the owner */
+	void SetOwnerWorldInteraction( class UViewportWorldInteraction* InWorldInteraction );
+
+	/** Gets the owner */
+	class UViewportWorldInteraction* GetOwnerWorldInteraction() const;
+
+	/** Called by the world interaction system after we've been spawned into the world, to allow
+	    us to create components and set everything up nicely for the selected objects that we'll be
+		used to manipulate */
+	virtual void UpdateGizmo( const EGizmoHandleTypes GizmoType, const ECoordSystem GizmoCoordinateSpace, const FTransform& LocalToWorld, const FBox& LocalBounds, const FVector ViewLocation, bool bAllHandlesVisible, 
+		class UActorComponent* DraggingHandle, const TArray< UActorComponent* >& HoveringOverHandles, const float GizmoHoverScale, const float GizmoHoverAnimationDuration ) {}
+
+protected:
+
+	/** Static: Given a bounding box and information about which edge to query, returns the vertex positions of that edge */
+	static void GetBoundingBoxEdge( const FBox& Box, const int32 AxisIndex, const int32 EdgeIndex, FVector& OutVertex0, FVector& OutVertex1 );
+
+	/** Updates the visibility of all the handles */
+	void UpdateHandleVisibility( const EGizmoHandleTypes GizmoType, const ECoordSystem GizmoCoordinateSpace, bool bAllHandlesVisible, UActorComponent* DraggingHandle );
+
+	/** Gets if the gizmo shows measurement texts */
+	bool GetShowMeasurementText() const;
+
+	/** Real time that the gizmo was last attached to a selected set of objects.  This is used for animation transitions */
+	FTimespan SelectedAtTime;
+	
+	/** Scene component root of this actor */
+	UPROPERTY()
+	USceneComponent* SceneComponent;
+
+	/** Gizmo material (opaque) */
+	UPROPERTY()
+	UMaterialInterface* GizmoMaterial;
+
+	/** Gizmo material (translucent) */
+	UPROPERTY()
+	UMaterialInterface* TranslucentGizmoMaterial;
+
+	/** All gizmo components */
+	UPROPERTY()
+	TArray< class UGizmoHandleGroup* > AllHandleGroups;
+
+	/** Owning object */
+	UPROPERTY()
+	class UViewportWorldInteraction* WorldInteraction;
+};	
\ No newline at end of file
diff --git a/Engine/Source/Editor/ViewportInteraction/Public/VIGizmoHandle.h b/Engine/Source/Editor/ViewportInteraction/Public/VIGizmoHandle.h
new file mode 100644
--- /dev/null
+++ b/Engine/Source/Editor/ViewportInteraction/Public/VIGizmoHandle.h
@@ -0,0 +1,142 @@
+// Copyright 1998-2015 Epic Games, Inc. All Rights Reserved.
+
+#pragma once
+
+#include "Components/SceneComponent.h"
+#include "ViewportInteractionTypes.h"
+#include "VIGizmoHandle.generated.h"
+
+enum class EGizmoHandleTypes : uint8;
+
+USTRUCT()
+struct VIEWPORTINTERACTION_API FGizmoHandle
+{
+	GENERATED_BODY()
+
+	/** Static mesh for this handle */
+	class UStaticMeshComponent* HandleMesh;
+
+	/** Scalar that will advance toward 1.0 over time as we hover over the gizmo handle */
+	float HoverAlpha;
+};
+
+
+/**
+ * Base class for gizmo handles
+ */
+UCLASS( ABSTRACT )
+class VIEWPORTINTERACTION_API UGizmoHandleGroup : public USceneComponent
+{
+	GENERATED_BODY()
+
+public:
+
+	/** Default constructor that sets up CDO properties */
+	UGizmoHandleGroup();
+	
+	/** Given the unique index, makes a handle */
+	FTransformGizmoHandlePlacement MakeHandlePlacementForIndex( const int32 HandleIndex ) const;
+
+	/** Makes a unique index for a handle */
+	int32 MakeHandleIndex( const FTransformGizmoHandlePlacement HandlePlacement ) const;
+
+	/** Makes up a name string for a handle */
+	FString MakeHandleName( const FTransformGizmoHandlePlacement HandlePlacement ) const;
+
+	/** Static: Given an axis (0-2) and a facing direction, returns the vector normal for that axis */
+	static FVector GetAxisVector( const int32 AxisIndex, const ETransformGizmoHandleDirection HandleDirection );
+
+	/** Updates the Gizmo handles, needs to be implemented by derived classes */
+	virtual void UpdateGizmoHandleGroup(  const FTransform& LocalToWorld, const FBox& LocalBounds, const FVector ViewLocation, bool bAllHandlesVisible, class UActorComponent* DraggingHandle, const TArray< UActorComponent* >& HoveringOverHandles, 
+		float AnimationAlpha, float GizmoScale, const float GizmoHoverScale, const float GizmoHoverAnimationDuration, bool& bOutIsHoveringOrDraggingThisHandleGroup );
+
+	/** Gets the InteractionType and the HandlePlacement for this Gizmo handle */
+	virtual void GetHandleIndexInteractionType( const int32 HandleIndex, ETransformGizmoInteractionType& OutInteractionType, TOptional<FTransformGizmoHandlePlacement>& OutHandlePlacement );
+
+	/** Gets the Gizmo InteractionType, needs to be implemented by derived classes */
+	virtual ETransformGizmoInteractionType GetInteractionType() const;
+
+	/** Finds the index of DraggedMesh in HandleMeshes */
+	virtual int32 GetDraggedHandleIndex( class UStaticMeshComponent* DraggedMesh );
+
+	/** Sets the Gizmo material */
+	void SetGizmoMaterial( UMaterialInterface* Material );
+	
+	/** Sets the translucent Gizmo material */
+	void SetTranslucentGizmoMaterial( UMaterialInterface* Material );
+
+	/** Gets all the handles */
+	TArray< FGizmoHandle >& GetHandles();
+
+	/** Gets the GizmoType for this Gizmo handle */
+	virtual EGizmoHandleTypes GetHandleType() const;
+
+	/** Returns true if this type of handle is allowed with world space gizmos */
+	virtual bool SupportsWorldCoordinateSpace() const
+	{
+		return true;
+	}
+
+	/** Sets if this handlegroup will be visible with the universal gizmo */
+	void SetShowOnUniversalGizmo( const bool bShowHandleUniversal );
+	
+	/** Gets if this handlegroup will be visible with the universal gizmo */
+	bool GetShowOnUniversalGizmo() const;
+
+protected:
+	/** Updates the colors of the dynamic material instances for the handle passed using its axis index */	
+	void UpdateHandleColor( const int32 AxisIndex, FGizmoHandle& Handle, class UActorComponent* DraggingHandle, const TArray< UActorComponent* >& HoveringOverHandles );
+
+	/** Helper function to create gizmo handle meshes */
+	class UStaticMeshComponent* CreateMeshHandle( class UStaticMesh* HandleMesh, const FString& ComponentName );
+
+	/** Creates handle meshcomponent and adds it to the Handles list */
+	class UStaticMeshComponent* CreateAndAddMeshHandle( class UStaticMesh* HandleMesh, const FString& ComponentName, const FTransformGizmoHandlePlacement& HandlePlacement );
+
+	/** Adds the HandleMeshComponent to the Handles list */
+	void AddMeshToHandles( class UStaticMeshComponent* HandleMeshComponent, const FTransformGizmoHandlePlacement& HandlePlacement );
+
+	/** Gets the handleplacement axes */
+	FTransformGizmoHandlePlacement GetHandlePlacement( const int32 X, const int32 Y, const int32 Z ) const;
+
+	/** Gizmo material (opaque) */
+	UPROPERTY()
+	UMaterialInterface* GizmoMaterial;
+
+	/** Gizmo material (translucent) */
+	UPROPERTY()
+	UMaterialInterface* TranslucentGizmoMaterial;
+	
+	/** All the StaticMeshes for this handle type */
+	UPROPERTY()
+	TArray< FGizmoHandle > Handles;
+
+private:
+
+	/** Updates the hover animation for the HoveringOverHandles */
+	void UpdateHoverAnimation( class UActorComponent* DraggingHandle, const TArray< UActorComponent* >& HoveringOverHandles, const float GizmoHoverAnimationDuration, bool& bOutIsHoveringOrDraggingThisHandleGroup );
+
+	/** If this handlegroup will be visible with the universal gizmo */
+	bool bShowOnUniversalGizmo;
+};
+
+/**
+ * Base class for gizmo handles on axis
+ */
+UCLASS(ABSTRACT)
+class VIEWPORTINTERACTION_API UAxisGizmoHandleGroup : public UGizmoHandleGroup
+{
+	GENERATED_BODY()
+
+protected:
+
+	/** Creates mesh for every axis */
+	void CreateHandles(UStaticMesh* HandleMesh, const FString& HandleComponentName);
+	
+	/** 
+	 * Calculates the transforms of meshes of this handlegroup 
+	 * @param HandleToCenter - The offset from the center of the gizmo in roomspace
+	 */
+	void UpdateHandlesRelativeTransformOnAxis( const FTransform& HandleToCenter, const float AnimationAlpha, const float GizmoScale, const float GizmoHoverScale, 
+		const FVector& ViewLocation, class UActorComponent* DraggingHandle, const TArray< UActorComponent* >& HoveringOverHandles );
+};
\ No newline at end of file
diff --git a/Engine/Source/Editor/ViewportInteraction/Public/VIStretchGizmoHandle.h b/Engine/Source/Editor/ViewportInteraction/Public/VIStretchGizmoHandle.h
new file mode 100644
--- /dev/null
+++ b/Engine/Source/Editor/ViewportInteraction/Public/VIStretchGizmoHandle.h
@@ -0,0 +1,33 @@
+// Copyright 1998-2015 Epic Games, Inc. All Rights Reserved.
+
+#pragma once
+
+#include "VIGizmoHandle.h"
+#include "VIStretchGizmoHandle.generated.h"
+
+/**
+ * Gizmo handle for stretching/scaling
+ */
+UCLASS()
+class VIEWPORTINTERACTION_API UStretchGizmoHandleGroup : public UGizmoHandleGroup
+{
+	GENERATED_BODY()
+
+public:
+
+	/** Default constructor that sets up CDO properties */
+	UStretchGizmoHandleGroup();
+	
+	/** Updates the Gizmo handles */
+	virtual void UpdateGizmoHandleGroup( const FTransform& LocalToWorld, const FBox& LocalBounds, const FVector ViewLocation, bool bAllHandlesVisible, class UActorComponent* DraggingHandle, const TArray< UActorComponent* >& HoveringOverHandles, 
+		float AnimationAlpha, float GizmoScale, const float GizmoHoverScale, const float GizmoHoverAnimationDuration, bool& bOutIsHoveringOrDraggingThisHandleGroup ) override;
+
+	/** Gets the InteractionType for this Gizmo handle */
+	virtual ETransformGizmoInteractionType GetInteractionType() const override;
+
+	/** Gets the GizmoType for this Gizmo handle */
+	virtual EGizmoHandleTypes GetHandleType() const override;
+
+	/** Returns true if this type of handle is allowed with world space gizmos */
+	virtual bool SupportsWorldCoordinateSpace() const override;
+};
\ No newline at end of file
diff --git a/Engine/Source/Editor/ViewportInteraction/Public/VIUniformScaleGizmoHandle.h b/Engine/Source/Editor/ViewportInteraction/Public/VIUniformScaleGizmoHandle.h
new file mode 100644
--- /dev/null
+++ b/Engine/Source/Editor/ViewportInteraction/Public/VIUniformScaleGizmoHandle.h
@@ -0,0 +1,38 @@
+// Copyright 1998-2015 Epic Games, Inc. All Rights Reserved.
+
+#pragma once
+
+#include "VIGizmoHandle.h"
+#include "VIUniformScaleGizmoHandle.generated.h"
+
+/**
+ * Gizmo handle for uniform scaling
+ */
+UCLASS()
+class VIEWPORTINTERACTION_API UUniformScaleGizmoHandleGroup : public UGizmoHandleGroup
+{
+	GENERATED_BODY()
+
+public:
+
+	/** Default constructor that sets up CDO properties */
+	UUniformScaleGizmoHandleGroup();
+	
+	/** Updates the Gizmo handles */
+	virtual void UpdateGizmoHandleGroup( const FTransform& LocalToWorld, const FBox& LocalBounds, const FVector ViewLocation, bool bAllHandlesVisible, class UActorComponent* DraggingHandle,
+		const TArray< UActorComponent* >& HoveringOverHandles, float AnimationAlpha, float GizmoScale, const float GizmoHoverScale, const float GizmoHoverAnimationDuration, bool& bOutIsHoveringOrDraggingThisHandleGroup ) override;
+
+	/** Gets the InteractionType for this Gizmo handle */
+	virtual ETransformGizmoInteractionType GetInteractionType() const override;
+
+	/** Gets the GizmoType for this Gizmo handle */
+	virtual EGizmoHandleTypes GetHandleType() const override;
+
+	/** Sets if the pivot point is used as location for the handle */
+	void SetUsePivotPointAsLocation( const bool bInUsePivotAsLocation );
+
+private:
+
+	/** If the pivot point is used for the uniform scaling handle */
+	bool bUsePivotAsLocation;
+};
\ No newline at end of file
diff --git a/Engine/Source/Editor/ViewportInteraction/Public/ViewportDragOperation.h b/Engine/Source/Editor/ViewportInteraction/Public/ViewportDragOperation.h
new file mode 100644
--- /dev/null
+++ b/Engine/Source/Editor/ViewportInteraction/Public/ViewportDragOperation.h
@@ -0,0 +1,61 @@
+// Copyright 1998-2015 Epic Games, Inc. All Rights Reserved.
+
+#pragma once
+
+#include "Components/ActorComponent.h"
+#include "ViewportDragOperation.generated.h"
+
+
+/**
+ * Base class for interactable drag calculations
+ */
+UCLASS()
+class VIEWPORTINTERACTION_API UViewportDragOperation : public UObject
+{
+	GENERATED_BODY()
+
+public:
+
+	/** 
+	 * Execute dragging 
+	 *
+	 * @param UViewportInteractor - The interactor causing the dragging
+	 * @param IViewportInteractableInterface - The interactable owning this drag operation
+	 */
+	virtual void ExecuteDrag( class UViewportInteractor* Interactor, class IViewportInteractableInterface* Interactable ) PURE_VIRTUAL( UViewportDragOperation::ExecuteDrag, );
+};
+
+/**
+ * Container component for UViewportDragOperation that can be used by objects in the world that are draggable and implement the ViewportInteractableInterface
+ */
+UCLASS()
+class VIEWPORTINTERACTION_API UViewportDragOperationComponent : public UActorComponent
+{
+	GENERATED_UCLASS_BODY()
+
+public:
+	/** Destructor */
+	~UViewportDragOperationComponent();
+
+	/** Get the actual dragging operation */
+	UViewportDragOperation* GetDragOperation();
+
+	/** Sets the drag operation class that will be used next time starting dragging */
+	void SetDragOperactionClass( const TSubclassOf<UViewportDragOperation> InDragOperation );
+
+	/** Starts new dragging operation */
+	void StartDragOperation();
+
+	/* Destroys the dragoperation */
+	void ClearDragOperation();
+
+private:
+
+	/** The actual dragging implementation */
+	UPROPERTY()
+	UViewportDragOperation* DragOperation;
+
+	/** The next class that will be used as drag operation */
+	UPROPERTY()
+	TSubclassOf<UViewportDragOperation> DragOperationSubclass;
+};
diff --git a/Engine/Source/Editor/ViewportInteraction/Public/ViewportInteraction.h b/Engine/Source/Editor/ViewportInteraction/Public/ViewportInteraction.h
new file mode 100644
--- /dev/null
+++ b/Engine/Source/Editor/ViewportInteraction/Public/ViewportInteraction.h
@@ -0,0 +1,9 @@
+// Copyright 1998-2015 Epic Games, Inc. All Rights Reserved.
+
+#pragma once
+
+#include "ViewportInteractableInterface.h"
+#include "IViewportInteractionModule.h"
+#include "ViewportWorldInteractionInterface.h"
+#include "ViewportInteraction.h"
+#include "ViewportInteractor.h"
\ No newline at end of file
diff --git a/Engine/Source/Editor/ViewportInteraction/Public/ViewportInteractionTypes.h b/Engine/Source/Editor/ViewportInteraction/Public/ViewportInteractionTypes.h
new file mode 100644
--- /dev/null
+++ b/Engine/Source/Editor/ViewportInteraction/Public/ViewportInteractionTypes.h
@@ -0,0 +1,126 @@
+// Copyright 1998-2015 Epic Games, Inc. All Rights Reserved.
+
+#pragma once
+
+#include "ViewportInteractionTypes.generated.h"
+
+/**
+ * Represents a generic action
+*/
+USTRUCT()
+struct VIEWPORTINTERACTION_API FViewportActionKeyInput
+{
+	GENERATED_BODY()
+	
+	FViewportActionKeyInput() :
+		ActionType( NAME_None ),
+		bIsInputCaptured( false )
+	{}
+
+	FViewportActionKeyInput( const FName& InActionType ) : 
+		ActionType( InActionType ),
+		bIsInputCaptured( false )
+	{}
+
+	/** The name of this action */
+	FName ActionType;
+
+	/** Input event */
+	EInputEvent Event;
+
+	/** True if this action owned by an interactor is "captured" for each possible action type, meaning that only the active captor should 
+	handle input events until it is no longer captured.  It's the captors responsibility to set this using OnVRAction(), or clear it when finished with capturing. */
+	bool bIsInputCaptured;
+};
+
+
+/** Methods of dragging objects around in VR */
+UENUM()
+enum class EViewportInteractionDraggingMode : uint8
+{
+	/** Not dragging right now with this hand */
+	Nothing,
+
+	/** Dragging actors around using the transform gizmo */
+	ActorsWithGizmo,
+
+	/** Actors locked to the impact point under the laser */
+	ActorsAtLaserImpact,
+
+	/** We're grabbing an object (or the world) that was already grabbed by the other hand */
+	AssistingDrag,
+
+	/** Freely moving, rotating and scaling objects with one or two hands */
+	ActorsFreely,
+
+	/** Moving the world itself around (actually, moving the camera in such a way that it feels like you're moving the world) */
+	World,
+
+	/** Moving a custom interactable */
+	Interactable,
+
+	/** Custom implementation for dragging */
+	Custom,
+
+	/** Dragging a material */
+	Material
+};
+
+/**
+* Things the transform gizmo can do to objects
+*/
+UENUM()
+enum class ETransformGizmoInteractionType : uint8
+{
+	/** No interaction */
+	None,
+
+	/** Translate the object(s), either in any direction or confined to one axis.  No rotation. */
+	Translate,
+
+	/** Translate the object, constrained to a plane.  No rotation. */
+	TranslateOnPlane,
+
+	/** Stretch the object non-uniformly while simultaneously repositioning it so account for it's new size */
+	StretchAndReposition,
+
+	/** Scaling the object non-uniformly while	 */
+	Scale,
+
+	/** Rotate the object(s) around a specific axis relative to the gizmo's pivot, translating as needed */
+	Rotate,
+
+	/** Rotate the object(s) around the pivot based on start angle */
+	RotateOnAngle,
+
+	/** Uniform scale the object(s) */
+	UniformScale,
+};
+
+
+/* Directions that a transform handle can face along any given axis */
+enum class ETransformGizmoHandleDirection
+{
+	Negative,
+	Center,
+	Positive,
+};
+
+
+/** Placement of a handle in pivot space */
+USTRUCT()
+struct VIEWPORTINTERACTION_API FTransformGizmoHandlePlacement
+{
+	GENERATED_BODY()
+
+	/* Handle direction in X, Y and Z */
+	ETransformGizmoHandleDirection Axes[ 3 ];
+
+
+	/** Finds the center handle count and facing axis index.  The center handle count is simply the number
+	of axes where the handle would be centered on the bounds along that axis.  The facing axis index is
+	the index (0-2) of the axis where the handle would be facing, or INDEX_NONE for corners or edges.
+	The center axis index is valid for edges, and defines the axis perpendicular to that edge direction,
+	or INDEX_NONE if it's not an edge */
+	void GetCenterHandleCountAndFacingAxisIndex( int32& OutCenterHandleCount, int32& OutFacingAxisIndex, int32& OutCenterAxisIndex ) const;
+};
diff --git a/Engine/Source/Editor/ViewportInteraction/Public/ViewportInteractor.h b/Engine/Source/Editor/ViewportInteraction/Public/ViewportInteractor.h
new file mode 100644
--- /dev/null
+++ b/Engine/Source/Editor/ViewportInteraction/Public/ViewportInteractor.h
@@ -0,0 +1,179 @@
+// Copyright 1998-2015 Epic Games, Inc. All Rights Reserved.
+
+#pragma once
+
+#include "ViewportInteractorData.h"
+#include "ViewportInteractor.generated.h"
+
+/**
+ * Represents the interactor in the world
+ */
+UCLASS( ABSTRACT )
+class VIEWPORTINTERACTION_API UViewportInteractor : public UObject
+{
+	GENERATED_UCLASS_BODY()
+
+public:
+
+	/** Gets the private data for this interactor */
+	struct FViewportInteractorData& GetInteractorData();
+
+	/** Sets the world interaction */
+	void SetWorldInteraction( class UViewportWorldInteraction* InWorldInteraction );
+
+	/** Sets the other interactor */
+	void SetOtherInteractor( UViewportInteractor* InOtherInteractor ); //@todo ViewportInteraction: This should not be public to other modules and should only be called from the world interaction.
+
+	/** Removes the other interactor reference for this interactor */
+	void RemoveOtherInteractor();
+
+	/** Gets the paired interactor assigned by the world interaction, can return null when there is no other interactor */
+	UViewportInteractor* GetOtherInteractor() const;
+
+	/** Updates the hover state of this interactor */
+	void UpdateHoverResult( FHitResult& HitResult );
+
+	/** Whenever the ViewportWorldInteraction is shut down, the interacts will shut down as well */
+	virtual void Shutdown();
+
+	/** Update for this interactor called by the ViewportWorldInteraction */
+	virtual void Tick( const float DeltaTime ) {};
+
+	/** Gets the hover component from the interactor data */
+	virtual class UActorComponent* GetHoverComponent();
+
+	/** Adds a new action to the KeyToActionMap */
+	void AddKeyAction( const FKey& Key, const FViewportActionKeyInput& Action );
+
+	/** Removes an action from the KeyToActionMap */
+	void RemoveKeyAction( const FKey& Key );
+
+	/** Base classes need to implement getting the input for the input devices for that interactor */
+	virtual void PollInput() {}
+
+	/** Handles axis input and translates it actions */
+	bool HandleInputKey( const FKey Key, const EInputEvent Event );
+
+	/** Handles axis input and translates it to actions */
+	bool HandleInputAxis( const FKey Key, const float Delta, const float DeltaTime );
+
+	/** Gets the world transform of this interactor */
+	FTransform GetTransform() const;
+
+	/** Gets the current interactor data dragging mode */
+	EViewportInteractionDraggingMode GetDraggingMode() const;
+
+	/** Gets the interactor data previous dragging mode */
+	EViewportInteractionDraggingMode GetLastDraggingMode() const;
+
+	/** Gets the interactor data drag velocity */
+	FVector GetDragTranslationVelocity() const;
+
+	/**
+	 * Gets the start and end point of the laser pointer for the specified hand
+	 *
+	 * @param HandIndex				Index of the hand to use
+	 * @param LasertPointerStart	(Out) The start location of the laser pointer in world space
+	 * @param LasertPointerEnd		(Out) The end location of the laser pointer in world space
+	 * @param bEvenIfUIIsInFront	If true, returns a laser pointer even if the hand has UI in front of it (defaults to false)
+	 * @param LaserLengthOverride	If zero the default laser length (VREdMode::GetLaserLength) is used
+	 *
+	 * @return	True if we have motion controller data for this hand and could return a valid result
+	 */
+	bool GetLaserPointer( FVector& LaserPointerStart, FVector& LaserPointerEnd, const bool bEvenIfBlocked = false, const float LaserLengthOverride = 0.0f );
+
+	/** Gets the maximum length of a laser pointer */
+	float GetLaserPointerMaxLength() const;
+
+	/** Triggers a force feedback effect on device if possible */
+	virtual void PlayHapticEffect( const float Strength ) {};
+
+	/**
+	 * Traces along the laser pointer vector and returns what it first hits in the world
+	 *
+	 * @param HandIndex	Index of the hand that uses has the laser pointer to trace
+	 * @param OptionalListOfIgnoredActors Actors to exclude from hit testing
+	 * @param bIgnoreGizmos True if no gizmo results should be returned, otherwise they are preferred (x-ray)
+	 * @param bEvenIfUIIsInFront If true, ignores any UI that might be blocking the ray
+	 * @param LaserLengthOverride If zero the default laser length (VREdMode::GetLaserLength) is used
+	 *
+	 * @return What the laster pointer hit
+	 */
+	virtual FHitResult GetHitResultFromLaserPointer( TArray<AActor*>* OptionalListOfIgnoredActors = nullptr, const bool bIgnoreGizmos = false,
+		TArray<UClass*>* ObjectsInFrontOfGizmo = nullptr, const bool bEvenIfBlocked = false, const float LaserLengthOverride = 0.0f );
+
+	/** Reset the values before checking the hover actions */
+	virtual void ResetHoverState( const float DeltaTime ) {};
+
+	/** Needs to be implemented by the base to calculate drag ray length and the velocity for the ray */
+	virtual void CalculateDragRay() {};
+
+	/**
+	 * Creates a hand transform and forward vector for a laser pointer for a given hand
+	 *
+	 * @param HandIndex			Index of the hand to use
+	 * @param OutHandTransform	The created hand transform
+	 * @param OutForwardVector	The forward vector of the hand
+	 *
+	 * @return	True if we have motion controller data for this hand and could return a valid result
+	 */
+	virtual bool GetTransformAndForwardVector( FTransform& OutHandTransform, FVector& OutForwardVector );
+
+	/** Called by StartDragging in world interaction to give the interator a chance of acting upon starting a drag operation */
+	virtual void OnStartDragging( UActorComponent* ClickedComponent, const FVector& HitLocation, const bool bIsPlacingActors ) {};
+
+	/** Gets the interactor laser hover location */
+	FVector GetHoverLocation() const;
+
+	/** If the interactor laser is currently hovering */
+	bool IsHovering() const;
+
+	/** Sets the current dragging mode for this interactor */
+	void SetDraggingMode( const EViewportInteractionDraggingMode NewDraggingMode );
+
+	/** To be overridden by base class. Called by GetLaserPointer to give the derived interactor a chance to disable the laser. By default it is not blocked */
+	virtual bool GetIsLaserBlocked() { return false; }
+
+	/** Gets a certain action by iterating through the map looking for the same ActionType */
+	FViewportActionKeyInput* GetActionWithName( const FName InActionName );
+
+	/** Gets the drag haptic feedback strength console variable */
+	float GetDragHapticFeedbackStrength() const;
+
+	/**
+	* Enables or disable "light press" locking for this interactor.  This is an advanced feature.
+	* It should be true if we allow locking of the lightly pressed state for the current event.  This can 
+	* be useful to turn off in cases where you always want a full press to override the light press, even 
+	* if it was a very slow press
+	*
+	* @param	bAllow		Enables or disables light press locking for this hand
+	*/
+	void SetAllowTriggerLightPressLocking( const bool bInAllow );
+
+	/** Check if this interactor allows the light press to be locked */
+	bool AllowTriggerLightPressLocking() const;
+
+protected:
+
+	/** To be overridden by base class. Called by HandleInputKey before delegates and default input implementation */
+	virtual void HandleInputKey( FViewportActionKeyInput& Action, const FKey Key, const EInputEvent Event, bool& bOutWasHandled ) {};
+	
+	/** To be overridden by base class. Called by HandleInputAxis before delegates and default input implementation */
+	virtual void HandleInputAxis( FViewportActionKeyInput& Action, const FKey Key, const float Delta, const float DeltaTime, bool& bOutWasHandled ) {};
+
+protected:
+
+	/** All the private data for the interactor */
+	FViewportInteractorData InteractorData;
+	
+	/** Mapping of raw keys to actions*/
+	TMap<FKey, FViewportActionKeyInput> KeyToActionMap;
+
+	/** The owning world interaction */
+	UPROPERTY()
+	class UViewportWorldInteraction* WorldInteraction;
+
+	/** The paired interactor by the world interaction */
+	UPROPERTY()
+	UViewportInteractor* OtherInteractor;
+};
\ No newline at end of file
diff --git a/Engine/Source/Editor/ViewportInteraction/Public/ViewportInteractorMotionController.h b/Engine/Source/Editor/ViewportInteraction/Public/ViewportInteractorMotionController.h
new file mode 100644
--- /dev/null
+++ b/Engine/Source/Editor/ViewportInteraction/Public/ViewportInteractorMotionController.h
@@ -0,0 +1,148 @@
+// Copyright 1998-2015 Epic Games, Inc. All Rights Reserved.
+
+#pragma once
+
+#include "ViewportInteractor.h"
+#include "HeadMountedDisplayTypes.h"	// For EHMDDeviceType::Type
+#include "ViewportInteractorMotionController.generated.h"
+
+/**
+ * Represents the interactor in the world
+ */
+UCLASS()
+class VIEWPORTINTERACTION_API UViewportInteractorMotionController : public UViewportInteractor
+{
+	GENERATED_UCLASS_BODY()
+
+public:
+
+	virtual ~UViewportInteractorMotionController();
+
+	// ViewportInteractorInterface overrides
+	virtual void Shutdown() override;
+	virtual void Tick( const float DeltaTime ) override;
+
+	/** @return Returns the type of HMD we're dealing with */
+	EHMDDeviceType::Type GetHMDDeviceType() const;
+
+	// ViewportInteractor
+	virtual bool GetTransformAndForwardVector( FTransform& OutHandTransform, FVector& OutForwardVector );
+
+	/** Starts haptic feedback for physical motion controller */
+	virtual void PlayHapticEffect( const float Strength ) override;
+	
+	/** Stops playing any haptic effects that have been going for a while.  Called every frame. */
+	void StopOldHapticEffects();
+	
+	/** Get the side of the controller */
+	EControllerHand GetControllerSide() const;
+
+	/** Get the motioncontroller component of this interactor */
+	class UMotionControllerComponent* GetMotionControllerComponent() const;
+	
+	/** Gets the trigger value */
+	float GetSelectAndMoveTriggerValue() const;
+
+protected:
+
+	// ViewportInteractor
+	virtual void HandleInputAxis( FViewportActionKeyInput& Action, const FKey Key, const float Delta, const float DeltaTime, bool& bOutWasHandled ) override;
+
+	/** Polls input for the motion controllers transforms */
+	virtual void PollInput() override;
+
+	/** Check if the trigger is lightly pressed */
+	bool IsTriggerAtLeastLightlyPressed() const;
+
+	/** Get the real time the trigger was lightly pressed */
+	double GetRealTimeTriggerWasLightlyPressed() const;
+
+	/** Set that the trigger is at least lighty pressed */
+	void SetTriggerAtLeastLightlyPressed( const bool bInTriggerAtLeastLightlyPressed );
+
+	/** Set that the trigger has been released since last press */
+	void SetTriggerBeenReleasedSinceLastPress( const bool bInTriggerBeenReleasedSinceLastPress );
+
+protected:
+	
+	/** Motion controller component which handles late-frame transform updates of all parented sub-components */
+	UPROPERTY()
+	class UMotionControllerComponent* MotionControllerComponent;
+
+	//
+	// Graphics
+	//
+
+	/** Mesh for this hand */
+	UPROPERTY()
+	class UStaticMeshComponent* HandMeshComponent;
+
+	/** Mesh for this hand's laser pointer */
+	UPROPERTY()
+	class UStaticMeshComponent* LaserPointerMeshComponent;
+
+	/** MID for laser pointer material (opaque parts) */
+	UPROPERTY()
+	class UMaterialInstanceDynamic* LaserPointerMID;
+
+	/** MID for laser pointer material (translucent parts) */
+	UPROPERTY()
+	class UMaterialInstanceDynamic* TranslucentLaserPointerMID;
+
+	/** Hover impact indicator mesh */
+	UPROPERTY()
+	class UStaticMeshComponent* HoverMeshComponent;
+
+	/** Hover point light */
+	UPROPERTY()
+	class UPointLightComponent* HoverPointLightComponent;
+
+	/** MID for hand mesh */
+	UPROPERTY()
+	class UMaterialInstanceDynamic* HandMeshMID;
+
+	/** Right or left hand */
+	EControllerHand ControllerHandSide;
+
+	/** True if this hand has a motion controller (or both!) */
+	bool bHaveMotionController;
+
+	// Special key action names for motion controllers
+	static const FName TrackpadPositionX;
+	static const FName TrackpadPositionY;
+	static const FName TriggerAxis;
+	static const FName MotionController_Left_FullyPressedTriggerAxis;
+	static const FName MotionController_Right_FullyPressedTriggerAxis;
+	static const FName MotionController_Left_LightlyPressedTriggerAxis;
+	static const FName MotionController_Right_LightlyPressedTriggerAxis;
+
+private:
+
+	//
+	// Trigger axis state
+	//
+
+	/** True if trigger is fully pressed right now (or within some small threshold) */
+	bool bIsTriggerFullyPressed;
+
+	/** True if the trigger is currently pulled far enough that we consider it in a "half pressed" state */
+	bool bIsTriggerAtLeastLightlyPressed;
+
+	/** Real time that the trigger became lightly pressed.  If the trigger remains lightly pressed for a reasonably 
+	long duration before being pressed fully, we'll continue to treat it as a light press in some cases */
+	double RealTimeTriggerWasLightlyPressed;
+
+	/** True if trigger has been fully released since the last press */
+	bool bHasTriggerBeenReleasedSinceLastPress;
+
+	/** Current trigger pressed amount for 'select and move' (0.0 - 1.0) */
+	float SelectAndMoveTriggerValue;
+
+	//
+	// Haptic feedback
+	//
+
+	/** The last real time we played a haptic effect.  This is used to turn off haptic effects shortly after they're triggered. */
+	double LastHapticTime;
+
+};
\ No newline at end of file
diff --git a/Engine/Source/Editor/ViewportInteraction/Public/ViewportWorldInteraction.h b/Engine/Source/Editor/ViewportInteraction/Public/ViewportWorldInteraction.h
new file mode 100644
--- /dev/null
+++ b/Engine/Source/Editor/ViewportInteraction/Public/ViewportWorldInteraction.h
@@ -0,0 +1,475 @@
+// Copyright 1998-2015 Epic Games, Inc. All Rights Reserved.
+
+#pragma once
+
+#include "ViewportWorldInteractionInterface.h"
+#include "ViewportInteractionTypes.h"
+#include "ViewportInteractionInputProcessor.h"
+#include "ViewportWorldInteraction.generated.h"
+
+namespace ViewportWorldActionTypes
+{
+	static const FName NoAction( "NoAction" );
+	static const FName WorldMovement( "WorldMovement" );
+	static const FName SelectAndMove( "SelectAndMove" );
+	static const FName SelectAndMove_LightlyPressed( "SelectAndMove_LightlyPressed" );
+	static const FName Undo( "Undo" );
+	static const FName Redo( "Redo" );
+	static const FName Delete( "Delete" );
+}
+
+// Forward declare the GizmoHandleTypes
+enum class EGizmoHandleTypes : uint8;
+class IViewportInteractableInterface;
+class UViewportInteractor;
+
+/**
+* Represents an object that we're actively interacting with, such as a selected actor
+*/
+class FViewportTransformable
+{
+
+public:
+
+	/** Sets up safe defaults */
+	FViewportTransformable()
+		: Object(),
+		StartTransform( FTransform::Identity ),
+		LastTransform( FTransform::Identity ),
+		TargetTransform( FTransform::Identity ),
+		UnsnappedTargetTransform( FTransform::Identity ),
+		InterpolationSnapshotTransform( FTransform::Identity )
+	{
+	}
+
+	/** The actual object */
+	FWeakObjectPtr Object;
+
+	/** The object's position when we started the action */
+	FTransform StartTransform;
+
+	/** The last position we were at.  Used for smooth interpolation when snapping is enabled */
+	FTransform LastTransform;
+
+	/** The target position we want to be at.  Used for smooth interpolation when snapping is enabled */
+	FTransform TargetTransform;
+
+	/** The target position we'd be at if snapping wasn't enabled.  This is used for 'elastic' smoothing when
+	snapping is turned on.  Basically it allows us to 'reach' toward the desired position */
+	FTransform UnsnappedTargetTransform;
+
+	/** A "snapshot" transform, used for certain interpolation effects */
+	FTransform InterpolationSnapshotTransform;
+};
+
+UCLASS()
+class VIEWPORTINTERACTION_API UViewportWorldInteraction : public UObject, public IViewportWorldInteractionInterface
+{
+	GENERATED_UCLASS_BODY()
+	
+public:
+
+	// IViewportWorldInteraction overrides
+	virtual void Init( class FEditorViewportClient* InEditorViewportClient ) override;
+	virtual void Shutdown() override;
+	virtual void Tick( class FEditorViewportClient* ViewportClient, const float DeltaTime ) override;
+	virtual void AddInteractor( UViewportInteractor* Interactor ) override;
+	virtual void RemoveInteractor( UViewportInteractor* Interactor ) override;
+	virtual FOnVIHoverUpdate& OnViewportInteractionHoverUpdate() override
+	{
+		return OnHoverUpdateEvent;
+	}
+	virtual FOnVIActionHandle& OnViewportInteractionInputAction() override
+	{
+		return OnInputActionEvent;
+	}
+
+	/** Pairs to interactors by setting	the other interactor for each interactor */
+	void PairInteractors( UViewportInteractor* FirstInteractor, UViewportInteractor* SecondInteractor );
+	
+	//
+	// Input
+ 	//
+
+	/** Handles the key input from the InputProcessor*/
+	bool HandleInputKey( const FKey Key, const EInputEvent Event );
+
+	/** Handles the axis input from the InputProcessor*/
+	bool HandleInputAxis( const int32 ControllerId, const FKey Key, const float Delta, const float DeltaTime );
+
+	/** Gets the world space transform of the calibrated VR room origin.  When using a seated VR device, this will feel like the
+	camera's world transform (before any HMD positional or rotation adjustments are applied.) */
+	FTransform GetRoomTransform() const;
+	
+	/** Gets the transform of the viewport / user's HMD in room space */
+	FTransform GetRoomSpaceHeadTransform() const;
+	
+	/** Gets the transform of the viewport / user's HMD in world space */
+	FTransform GetHeadTransform() const;
+
+	/** Sets a new transform for the room, in world space.  This is basically setting the editor's camera transform for the viewport */
+	void SetRoomTransform( const FTransform& NewRoomTransform );
+	
+	/** Gets the world scale factor, which can be multiplied by a scale vector to convert to room space */
+	float GetWorldScaleFactor() const;
+
+	/** Gets the world of the viewport used for this worldinteraction */
+	UWorld* GetViewportWorld() const;
+
+	/** Gets the currently used viewport */
+	FEditorViewportClient* GetViewportClient() const;
+
+	/** Returns the Unreal controller ID for the motion controllers we're using */
+	int32 GetMotionControllerID() const
+	{
+		return MotionControllerID;
+	}
+
+	/** Sets that the worldinteractor is used in VR, to give a correct head transform */
+	void SetIsUsingVR( const bool bInIsUsingVR );
+
+	/** Invokes the editor's undo system to undo the last change */
+	void Undo();
+
+	/** Redoes the last change */
+	void Redo();
+
+	/** Deletes all the selected objects */
+	void DeleteSelectedObjects();
+
+	/** Copies selected objects to the clipboard */
+	void Copy();
+
+	/** Pastes the clipboard contents into the scene */
+	void Paste();
+
+	/** Duplicates the selected objects */
+	void Duplicate();
+
+	/** Deselects all the current selected objects */
+	void Deselect();
+	
+	/** Called to finish a drag action with the specified interactor */
+	virtual void StopDragging( class UViewportInteractor* Interactor );
+
+	/** Starts dragging selected actors around, changing selection if needed.  Called when clicking in the world, or when placing new objects */
+	void StartDraggingActors( UViewportInteractor* Interactor, const FViewportActionKeyInput& Action, UActorComponent* ClickedComponent, const FVector& HitLocation, const bool bIsPlacingActors );
+	
+	/** Sets which transform gizmo coordinate space is used */
+	void SetTransformGizmoCoordinateSpace( const ECoordSystem NewCoordSystem );
+
+	/** Returns which transform gizmo coordinate space we're using, world or local */
+	ECoordSystem GetTransformGizmoCoordinateSpace() const;
+
+	/** Returns the maximum user scale */
+	float GetMaxScale();
+
+	/** Returns the minimum user scale */
+	float GetMinScale();
+
+	/** Sets GNewWorldToMetersScale */
+	void SetWorldToMetersScale( const float NewWorldToMetersScale );
+
+	/** True if we're actively interacting with the specified actor */
+	bool IsTransformingActor( AActor* Actor ) const;
+
+	/** Checks to see if we're allowed to interact with the specified component */
+	virtual bool IsInteractableComponent( const UActorComponent* Component ) const;
+
+	/** Gets the transform gizmo actor, or returns null if we currently don't have one */
+	class ABaseTransformGizmo* GetTransformGizmoActor();
+
+	/** Sets if objects are dragged with either hand since last time selecting something */
+	void SetDraggedSinceLastSelection( const bool bNewDraggedSinceLastSelection );
+
+	/** Sets the transform of the gizmo when starting drag */
+	void SetLastDragGizmoStartTransform( const FTransform NewLastDragGizmoStartTransform );
+
+	/** Gets all the interactors */
+	TArray<UViewportInteractor*>& GetInteractors();
+
+	/** Given a world space velocity vector, applies inertial damping to it to slow it down */
+	void ApplyVelocityDamping( FVector& Velocity, const bool bVelocitySensitive );
+
+	/** Gets the current Gizmo handle type */
+	EGizmoHandleTypes GetCurrentGizmoType() const;
+
+	/** Sets the current gizmo handle type */
+	void SetGizmoHandleType( const EGizmoHandleTypes InGizmoHandleType );
+
+	/** Changes the transform gizmo class that it will change to next refresh */
+	void SetTransformGizmoClass( const TSubclassOf<ABaseTransformGizmo>& NewTransformGizmoClass );
+
+	/** Sets the currently dragged interactavle */
+	void SetDraggedInteractable( IViewportInteractableInterface* InDraggedInteractable );	
+
+	/** Check if there is another interactor hovering over the component */
+	bool IsOtherInteractorHoveringOverComponent( UViewportInteractor* Interactor, UActorComponent* Component ) const;
+	
+	/** Switch which transform gizmo coordinate space we're using. */
+	void CycleTransformGizmoCoordinateSpace();
+
+protected:
+
+	/** Allow base class to refresh on selection changed */
+	virtual void RefreshOnSelectionChanged( const bool bNewObjectsSelected, bool bAllHandlesVisible, class UActorComponent* SingleVisibleHandle, const TArray< UActorComponent* >& HoveringOverHandles );
+
+	/** Creates a Transformable for each selected actor, wiping out the existing transformables */
+	void SetupTransformablesForSelectedActors();
+
+	/** Figures out where to place an object when tracing it against the scene using a laser pointer */
+	bool FindPlacementPointUnderLaser( UViewportInteractor* Interactor, FVector& OutHitLocation );
+
+private:
+
+	/** Called every frame to update hover state */
+	void HoverTick( class FEditorViewportClient* ViewportClient, const float DeltaTime );
+
+	/** Updates the interaction */
+	void InteractionTick( FEditorViewportClient* ViewportClient, const float DeltaTime );
+	
+	/** Called by the world interaction system when one of our components is dragged by the user.  If null is
+	passed in then we'll treat it as dragging the whole object (rather than a specific axis/handle) */
+	void UpdateDragging( 
+		const float DeltaTime,
+		bool& bIsFirstDragUpdate, 
+		const EViewportInteractionDraggingMode DraggingMode, 
+		const ETransformGizmoInteractionType InteractionType, 
+		const bool bWithTwoHands, 
+		class FEditorViewportClient& ViewportClient, 
+		const TOptional<FTransformGizmoHandlePlacement> OptionalHandlePlacement, 
+		const FVector& DragDelta, 
+		const FVector& OtherHandDragDelta, 
+		const FVector& DraggedTo, 
+		const FVector& OtherHandDraggedTo, 
+		const FVector& DragDeltaFromStart, 
+		const FVector& OtherHandDragDeltaFromStart, 
+		const FVector& LaserPointerStart, 
+		const FVector& LaserPointerDirection, 
+		const bool bIsLaserPointerValid, 
+		const FTransform& GizmoStartTransform, 
+		const FBox& GizmoStartLocalBounds, 
+		const USceneComponent* const DraggingTransformGizmoComponent,
+		FVector& GizmoSpaceFirstDragUpdateOffsetAlongAxis,
+		FVector& DragDeltaFromStartOffset,
+		bool& bIsDrivingVelocityOfSimulatedTransformables,
+		FVector& OutUnsnappedDraggedTo );
+	
+	/** Given a drag delta from a starting point, contrains that delta based on a gizmo handle axis */
+	FVector ComputeConstrainedDragDeltaFromStart( const bool bIsFirstDragUpdate, const ETransformGizmoInteractionType InteractionType, const TOptional<FTransformGizmoHandlePlacement> OptionalHandlePlacement, const FVector& DragDeltaFromStart, const FVector& LaserPointerStart, const FVector& LaserPointerDirection, const bool bIsLaserPointerValid, const FTransform& GizmoStartTransform, FVector& GizmoSpaceFirstDragUpdateOffsetAlongAxis, FVector& DragDeltaFromStartOffset ) const;
+
+	/** Checks to see if smooth snapping is currently enabled */
+	bool IsSmoothSnappingEnabled() const;
+
+	/** Returns the time since this was last entered */
+	inline FTimespan GetTimeSinceEntered() const
+	{
+		return FTimespan::FromSeconds( FApp::GetCurrentTime() ) - AppTimeEntered;
+	}
+
+	/** Updates the transforma */
+	void UpdateTransformables( const float DeltaTime );
+
+	/** Conditionally polls input from interactors, if we haven't done that already this frame */
+	void PollInputIfNeeded();
+
+	/** Called when selection changes in the level */
+	void OnActorSelectionChanged( UObject* ChangedObject );
+
+	/** Refreshes the transform gizmo to match the currently selected actor set */
+	void RefreshTransformGizmo( const bool bNewObjectsSelected, bool bAllHandlesVisible, class UActorComponent* SingleVisibleHandle, const TArray< UActorComponent* >& HoveringOverHandles );
+
+	/** Spawn new transform gizmo when refreshing and at start if there is none yet or the current one is different from the TransformGizmoClass */
+	void SpawnTransformGizmoIfNeeded();
+
+	/** Gets the inertie contribution of another interactor */
+	UViewportInteractor* GetOtherInteractorIntertiaContribute( UViewportInteractor* Interactor );
+
+	/** Spawns a transient actor that we can use in the editor world (templated for convenience) */
+	template<class T>
+	T* SpawnTransientSceneActor( const FString& ActorName, const bool bWithSceneComponent ) const
+	{
+		return CastChecked<T>( SpawnTransientSceneActor( T::StaticClass(), ActorName, bWithSceneComponent ) );
+	}
+
+	/** Spawns a transient actor that we can use in the editor world */
+	AActor* SpawnTransientSceneActor( TSubclassOf<AActor> ActorClass, const FString& ActorName, const bool bWithSceneComponent ) const;
+
+	/** Destroys a transient actor we created earlier and nulls out the pointer */
+	void DestroyTransientActor( AActor* Actor ) const;
+
+	/** Gets the snap grid mesh actor */
+	AActor* GetSnapGridActor()
+	{
+		SpawnGridMeshActor();
+		return SnapGridActor;
+	}
+
+	/** Gets the snap grid mesh MID */
+	class UMaterialInstanceDynamic* GetSnapGridMID()
+	{
+		SpawnGridMeshActor();
+		return SnapGridMID;
+	}
+
+	/** Spawns the grid actor */
+	void SpawnGridMeshActor();
+
+	//
+	// Colors
+	//
+public:
+	enum EColors
+	{
+		DefaultColor,
+		Forward,
+		Right,
+		Up,
+		Hover,
+		Dragging,
+		TotalCount
+	};
+
+	/** Gets the color from color type */
+	FLinearColor GetColor( const EColors Color ) const;
+
+private:
+
+	// All the colors for this mode
+	TArray<FLinearColor> Colors;
+
+protected:
+	
+	/** All of the objects we're currently interacting with, such as selected actors */
+	TArray< TUniquePtr< class FViewportTransformable > > Transformables;
+
+	/** True if we've dragged objects with either hand since the last time we selected something */
+	bool bDraggedSinceLastSelection;
+
+	/** Gizmo start transform of the last drag we did with either hand.  Only valid if bDraggedSinceLastSelection is true */
+	FTransform LastDragGizmoStartTransform;
+
+	//
+	// Undo/redo
+	//
+
+	/** Manages saving undo for selected actors while we're dragging them around */
+	FTrackingTransaction TrackingTransaction;
+
+private:
+
+	/** App time that we entered this */
+	FTimespan AppTimeEntered;
+
+	/** All the interactors registered to modify the world */
+	UPROPERTY()
+	TArray< class UViewportInteractor* > Interactors;
+
+	//
+	// Viewport
+	//
+
+	/** The viewport that the worldinteraction is used for */
+	class FEditorViewportClient* EditorViewportClient;
+
+	/** The last frame that input was polled.  We keep track of this so that we can make sure to poll for latest input right before
+	its first needed each frame */
+	uint32 LastFrameNumberInputWasPolled;
+
+	//
+	// Input
+	//
+
+	/** Slate Input Processor */
+	TSharedPtr<class FViewportInteractionInputProcessor> InputProcessor;
+
+	/** The Unreal controller ID for the motion controllers we're using */
+	int32 MotionControllerID;
+
+
+	//
+	// World movement
+	//
+
+	/** The world to meters scale before we changed it (last frame's world to meters scale) */
+	float LastWorldToMetersScale;
+
+	//
+	// Hover state
+	//
+
+	/** Set of objects that are currently hovered over */
+	TSet<FViewportHoverTarget> HoveredObjects;
+
+	//
+	// Object movement/placement
+	//
+
+	/** True if our transformables are currently interpolating to their target location */
+	bool bIsInterpolatingTransformablesFromSnapshotTransform;
+
+	/** Time that we started interpolating at */
+	FTimespan TransformablesInterpolationStartTime;
+
+	/** Duration to interpolate transformables over */
+	float TransformablesInterpolationDuration;
+
+	/** Event that fires every frame to update hover state based on what's under the cursor.  Set bWasHandled to true if you detect something to hover. */
+	FOnVIHoverUpdate OnHoverUpdateEvent;
+
+	/** Event that is fired for when a key is pressed by an interactor */
+	FOnVIActionHandle OnInputActionEvent;
+
+	//
+	// Transform gizmo
+	//
+
+	/** Transform gizmo actor, for manipulating selected actor(s) */
+	UPROPERTY()
+	class ABaseTransformGizmo* TransformGizmoActor;
+
+	/** Class used for when refreshing the transform gizmo */
+	TSubclassOf<ABaseTransformGizmo> TransformGizmoClass;
+
+	/** Current gizmo bounds in local space.  Updated every frame.  This is not the same as the gizmo actor's bounds. */
+	FBox GizmoLocalBounds;
+
+	/** Starting angle when rotating an object with  ETransformGizmoInteractionType::RotateOnAngle */
+	TOptional<float> StartDragAngleOnRotation;
+
+	/** The direction of where the rotation handle is facing when starting rotation */
+	TOptional<FVector> StartDragHandleDirection;
+
+	/** The current Gizmo type that is used for the TransformGizmo Actor */
+	EGizmoHandleTypes CurrentGizmoType;
+
+	//
+	// VR
+	//
+
+	/** If we are in VR right now */
+	bool bIsUsingVR;
+
+	//
+	// Snap grid
+	//
+
+	/** Actor for the snap grid */
+	UPROPERTY()
+	class AActor* SnapGridActor;
+
+	/** The plane mesh we use to draw a snap grid under selected objects */
+	UPROPERTY()
+	class UStaticMeshComponent* SnapGridMeshComponent;
+
+	/** MID for the snap grid, so we can update snap values on the fly */
+	UPROPERTY()
+	class UMaterialInstanceDynamic* SnapGridMID;
+
+	//
+	// Interactable
+	//
+
+	/** The current dragged interactable */
+	class IViewportInteractableInterface* DraggedInteractable;	
+};
\ No newline at end of file
diff --git a/Engine/Source/Editor/ViewportInteraction/Public/ViewportWorldInteractionInterface.h b/Engine/Source/Editor/ViewportInteraction/Public/ViewportWorldInteractionInterface.h
new file mode 100644
--- /dev/null
+++ b/Engine/Source/Editor/ViewportInteraction/Public/ViewportWorldInteractionInterface.h
@@ -0,0 +1,41 @@
+// Copyright 1998-2015 Epic Games, Inc. All Rights Reserved.
+
+#pragma once
+
+#include "ViewportWorldInteractionInterface.generated.h"
+
+UINTERFACE( MinimalAPI )
+class UViewportWorldInteractionInterface : public UInterface
+{
+	GENERATED_UINTERFACE_BODY()
+};
+
+class VIEWPORTINTERACTION_API IViewportWorldInteractionInterface
+{
+	GENERATED_IINTERFACE_BODY()
+
+public:
+	/** Initializes the world interaction */
+	virtual void Init( class FEditorViewportClient* InEditorViewportClient ) = 0;
+
+	/** Shuts down the world interaction */
+	virtual void Shutdown() = 0;
+
+	/** Main update loop including updating interaction, actions and hover */
+	virtual void Tick( class FEditorViewportClient* ViewportClient, const float DeltaTime ) = 0;
+	
+	/** Adds interactor to the worldinteraction */	
+	virtual void AddInteractor( class UViewportInteractor* Interactor ) = 0;
+
+	/** Removes interactor from the worldinteraction and removes the interactor from its paired interactor if any */
+	virtual void RemoveInteractor( class UViewportInteractor* Interactor ) = 0;
+
+	/** Gets the event for hovering update which is broadcasted for each interactor */
+	DECLARE_EVENT_FourParams( IViewportWorldInteractionInterface, FOnVIHoverUpdate, class FEditorViewportClient& /* ViewportClient */, class UViewportInteractor* /* Interactor */, FVector& /* OutHoverImpactPoint */,  bool& /* bWasHandled */ );
+	virtual FOnVIHoverUpdate& OnViewportInteractionHoverUpdate() = 0;
+
+	/** Gets the event for input actions update which is broadcasted for each interactor */
+	DECLARE_EVENT_SixParams( IViewportWorldInteractionInterface, FOnVIActionHandle, class FEditorViewportClient& /* ViewportClient */, class UViewportInteractor* /* Interactor */, 
+	const struct FViewportActionKeyInput& /* Action */, const EInputEvent /* Event */, bool& /* bOutIsInCaptured */, bool& /* bWasHandled */ );
+	virtual FOnVIActionHandle& OnViewportInteractionInputAction() = 0;
+};
\ No newline at end of file
diff --git a/Engine/Source/Editor/ViewportInteraction/ViewportDragOperation.cpp b/Engine/Source/Editor/ViewportInteraction/ViewportDragOperation.cpp
new file mode 100644
--- /dev/null
+++ b/Engine/Source/Editor/ViewportInteraction/ViewportDragOperation.cpp
@@ -0,0 +1,47 @@
+// Copyright 1998-2015 Epic Games, Inc. All Rights Reserved.
+
+#include "ViewportInteractionModule.h"
+#include "ViewportDragOperation.h"
+
+UViewportDragOperationComponent::UViewportDragOperationComponent( const class FObjectInitializer& Initializer ) :
+	Super( Initializer ),
+	DragOperation( nullptr )
+{
+
+}
+
+UViewportDragOperationComponent::~UViewportDragOperationComponent()
+{
+	DragOperation = nullptr;
+}
+
+UViewportDragOperation* UViewportDragOperationComponent::GetDragOperation()
+{
+	return DragOperation;
+}
+
+void UViewportDragOperationComponent::SetDragOperactionClass( const TSubclassOf<UViewportDragOperation> InDragOperation )
+{
+	DragOperationSubclass = InDragOperation;
+}
+
+void UViewportDragOperationComponent::StartDragOperation()
+{
+	if ( DragOperationSubclass )
+	{
+		// Reset the drag operation to make sure we start a new one
+		ClearDragOperation();
+		// Create the drag object with the latest class
+		DragOperation = NewObject<UViewportDragOperation>( ( UObject* ) GetTransientPackage(), DragOperationSubclass );
+	}
+}
+
+void UViewportDragOperationComponent::ClearDragOperation()
+{
+	if ( DragOperation )
+	{
+		DragOperation->MarkPendingKill();
+	}
+
+	DragOperation = nullptr;
+}
diff --git a/Engine/Source/Editor/ViewportInteraction/ViewportInteractableInterface.cpp b/Engine/Source/Editor/ViewportInteraction/ViewportInteractableInterface.cpp
new file mode 100644
--- /dev/null
+++ b/Engine/Source/Editor/ViewportInteraction/ViewportInteractableInterface.cpp
@@ -0,0 +1,10 @@
+// Copyright 1998-2015 Epic Games, Inc. All Rights Reserved.
+
+#include "ViewportInteractionModule.h"
+#include "ViewportInteractableInterface.h"
+
+UViewportInteractableInterface::UViewportInteractableInterface( const class FObjectInitializer& Initializer )
+	: Super( Initializer )
+{
+
+}
diff --git a/Engine/Source/Editor/ViewportInteraction/ViewportInteraction.Build.cs b/Engine/Source/Editor/ViewportInteraction/ViewportInteraction.Build.cs
new file mode 100644
--- /dev/null
+++ b/Engine/Source/Editor/ViewportInteraction/ViewportInteraction.Build.cs
@@ -0,0 +1,35 @@
+// Copyright 1998-2016 Epic Games, Inc. All Rights Reserved.
+
+using UnrealBuildTool;
+
+namespace UnrealBuildTool.Rules
+{
+    public class ViewportInteraction : ModuleRules
+    {
+        public ViewportInteraction(TargetInfo Target)
+        {
+            PrivateDependencyModuleNames.AddRange(
+                new string[] {
+                    "Core",
+                    "CoreUObject",
+                    "Engine",
+                    "InputCore",
+                    "UnrealEd",
+                    "Slate",
+                    "SlateCore",
+                    "HeadMountedDisplay"
+                }
+            );
+
+            PrivateIncludePathModuleNames.AddRange(
+                new string[] {
+                }
+            );
+
+            DynamicallyLoadedModuleNames.AddRange(
+                new string[] {
+                }
+            );
+        }
+    }
+}
\ No newline at end of file
diff --git a/Engine/Source/Editor/ViewportInteraction/ViewportInteractionInputProcessor.cpp b/Engine/Source/Editor/ViewportInteraction/ViewportInteractionInputProcessor.cpp
new file mode 100644
--- /dev/null
+++ b/Engine/Source/Editor/ViewportInteraction/ViewportInteractionInputProcessor.cpp
@@ -0,0 +1,39 @@
+// Copyright 1998-2015 Epic Games, Inc. All Rights Reserved.
+
+#include "ViewportInteractionModule.h"
+#include "ViewportInteractionInputProcessor.h"
+#include "ViewportWorldInteraction.h"
+#include "SLevelViewport.h"
+
+FViewportInteractionInputProcessor::FViewportInteractionInputProcessor( UViewportWorldInteraction& InWorldInteraction  )
+	: WorldInteraction( InWorldInteraction  )
+{
+}
+
+FViewportInteractionInputProcessor::~FViewportInteractionInputProcessor()
+{
+}
+
+void FViewportInteractionInputProcessor::Tick( const float DeltaTime, FSlateApplication& SlateApp, TSharedRef<ICursor> Cursor )
+{
+}
+
+bool FViewportInteractionInputProcessor::HandleKeyDownEvent( FSlateApplication& SlateApp, const FKeyEvent& InKeyEvent )
+{
+	return WorldInteraction.HandleInputKey( InKeyEvent.GetKey(), InKeyEvent.IsRepeat() ? IE_Repeat : IE_Pressed );
+}
+
+bool FViewportInteractionInputProcessor::HandleKeyUpEvent( FSlateApplication& SlateApp, const FKeyEvent& InKeyEvent )
+{
+	return WorldInteraction.HandleInputKey( InKeyEvent.GetKey(), IE_Released );
+}
+
+bool FViewportInteractionInputProcessor::HandleAnalogInputEvent( FSlateApplication& SlateApp, const FAnalogInputEvent& InAnalogInputEvent )
+{
+	return WorldInteraction.HandleInputAxis( InAnalogInputEvent.GetUserIndex(), InAnalogInputEvent.GetKey(), InAnalogInputEvent.GetAnalogValue(), FApp::GetDeltaTime() );
+}
+
+bool FViewportInteractionInputProcessor::HandleMouseMoveEvent( FSlateApplication& SlateApp, const FPointerEvent& MouseEvent )
+{
+	return false;
+}
diff --git a/Engine/Source/Editor/ViewportInteraction/ViewportInteractionModule.cpp b/Engine/Source/Editor/ViewportInteraction/ViewportInteractionModule.cpp
new file mode 100644
--- /dev/null
+++ b/Engine/Source/Editor/ViewportInteraction/ViewportInteractionModule.cpp
@@ -0,0 +1,15 @@
+// Copyright 1998-2015 Epic Games, Inc. All Rights Reserved.
+
+#include "ViewportInteractionModule.h"
+#include "ViewportWorldInteraction.h"
+
+FViewportInteractionModule::FViewportInteractionModule()
+{
+}
+
+FViewportInteractionModule::~FViewportInteractionModule()
+{
+}
+
+IMPLEMENT_MODULE( FViewportInteractionModule, ViewportInteraction )
+
diff --git a/Engine/Source/Editor/ViewportInteraction/ViewportInteractionModule.h b/Engine/Source/Editor/ViewportInteraction/ViewportInteractionModule.h
new file mode 100644
--- /dev/null
+++ b/Engine/Source/Editor/ViewportInteraction/ViewportInteractionModule.h
@@ -0,0 +1,14 @@
+// Copyright 1998-2015 Epic Games, Inc. All Rights Reserved.
+
+#pragma once
+
+#include "UnrealEd.h"
+#include "IViewportInteractionModule.h"
+
+class FViewportInteractionModule : public IViewportInteractionModule
+{
+public:
+	
+	FViewportInteractionModule();
+	virtual ~FViewportInteractionModule();
+};
\ No newline at end of file
