VR Editor:

diff --git a/Engine/Source/Editor/ViewportInteraction/ViewportInteractor.cpp b/Engine/Source/Editor/ViewportInteraction/ViewportInteractor.cpp
new file mode 100644
--- /dev/null
+++ b/Engine/Source/Editor/ViewportInteraction/ViewportInteractor.cpp
@@ -0,0 +1,584 @@
+// Copyright 1998-2015 Epic Games, Inc. All Rights Reserved.
+
+#include "ViewportInteractionModule.h"
+#include "ViewportInteractor.h"
+#include "ViewportWorldInteraction.h"
+#include "VIBaseTransformGizmo.h"
+#include "ViewportInteractableInterface.h"
+
+namespace VI
+{
+	static FAutoConsoleVariable LaserPointerMaxLength( TEXT( "VI.LaserPointerMaxLength" ), 10000.0f, TEXT( "Maximum length of the laser pointer line" ) );
+	static FAutoConsoleVariable DragHapticFeedbackStrength( TEXT( "VI.DragHapticFeedbackStrength" ), 1.0f, TEXT( "Default strength for haptic feedback when starting to drag objects" ) ); //@todo ViewportInteraction: Duplicate from ViewportWorldInteraction
+	static FAutoConsoleVariable SelectionHapticFeedbackStrength( TEXT( "VI.SelectionHapticFeedbackStrength" ), 0.5f, TEXT( "Default strength for haptic feedback when selecting objects" ) );
+}
+
+UViewportInteractor::UViewportInteractor( const FObjectInitializer& Initializer ) : 
+	UObject( Initializer ),
+	InteractorData(),
+	KeyToActionMap(),
+	WorldInteraction( nullptr ),
+	OtherInteractor( nullptr )
+{
+
+}
+
+FViewportInteractorData& UViewportInteractor::GetInteractorData()
+{
+	return InteractorData;
+}
+
+void UViewportInteractor::SetWorldInteraction( UViewportWorldInteraction* InWorldInteraction )
+{
+	WorldInteraction = InWorldInteraction;
+}
+
+void UViewportInteractor::SetOtherInteractor( UViewportInteractor* InOtherInteractor )
+{
+	OtherInteractor = InOtherInteractor;
+}
+
+void UViewportInteractor::RemoveOtherInteractor()
+{
+	OtherInteractor = nullptr;
+}
+
+UViewportInteractor* UViewportInteractor::GetOtherInteractor() const
+{
+	return OtherInteractor;
+}
+
+void UViewportInteractor::UpdateHoverResult( FHitResult& OutHitResult )
+{
+	FHitResult HitResult = GetHitResultFromLaserPointer();
+	InteractorData.bIsHovering = HitResult.bBlockingHit;
+	InteractorData.HoverLocation = HitResult.ImpactPoint;
+	OutHitResult = HitResult;
+}
+
+void UViewportInteractor::Shutdown()
+{
+	KeyToActionMap.Empty();
+
+	WorldInteraction = nullptr;
+	OtherInteractor = nullptr;
+}
+
+class UActorComponent* UViewportInteractor::GetHoverComponent()
+{
+	return InteractorData.HoveredActorComponent.Get();
+}
+
+void UViewportInteractor::AddKeyAction( const FKey& Key, const FViewportActionKeyInput& Action )
+{
+	KeyToActionMap.Add( Key, Action );
+}
+
+void UViewportInteractor::RemoveKeyAction( const FKey& Key )
+{
+	KeyToActionMap.Remove( Key );
+}
+
+bool UViewportInteractor::HandleInputKey( const FKey Key, const EInputEvent Event )
+{
+	bool bHandled = false;
+	FViewportActionKeyInput* Action = KeyToActionMap.Find( Key );
+	if ( Action != nullptr )	// Ignore key repeats
+	{
+		Action->Event = Event;
+		
+		// Give subsystems a chance to handle actions for this interactor
+		WorldInteraction->OnViewportInteractionInputAction().Broadcast( *WorldInteraction->GetViewportClient(), this, *Action, Event, Action->bIsInputCaptured, bHandled );
+
+		if ( !bHandled )
+		{
+			// Give the derived classes a chance to update according to the input
+			HandleInputKey( *Action, Key, Event, bHandled );
+		}
+
+		// Start checking on default action implementation
+		if ( !bHandled )
+		{
+			// Selection/Movement
+			if ( ( Action->ActionType == ViewportWorldActionTypes::SelectAndMove ||
+				Action->ActionType == ViewportWorldActionTypes::SelectAndMove_LightlyPressed ) )
+			{
+				FViewportInteractorData& OtherInteractorData = OtherInteractor->GetInteractorData();
+
+				if ( Event == IE_Pressed )
+				{
+					// No clicking while we're dragging the world.  (No laser pointers are visible, anyway.)
+					const bool bIsDraggingWorldWithTwoHands =
+						( InteractorData.DraggingMode == EViewportInteractionDraggingMode::World && GetOtherInteractor()->GetInteractorData().DraggingMode == EViewportInteractionDraggingMode::AssistingDrag ) ||
+						( GetOtherInteractor()->GetInteractorData().DraggingMode == EViewportInteractionDraggingMode::World && InteractorData.DraggingMode == EViewportInteractionDraggingMode::AssistingDrag );
+
+					FHitResult HitResult = GetHitResultFromLaserPointer();
+					if ( !bIsDraggingWorldWithTwoHands && !bHandled && HitResult.Actor.IsValid() )
+					{
+						if ( WorldInteraction->IsInteractableComponent( HitResult.GetComponent() ) )
+						{
+							InteractorData.ClickingOnComponent = HitResult.GetComponent();
+
+							AActor* Actor = HitResult.Actor.Get();
+
+							FVector LaserPointerStart, LaserPointerEnd;
+							const bool bHaveLaserPointer = GetLaserPointer( /* Out */ LaserPointerStart, /* Out */ LaserPointerEnd );
+							check( bHaveLaserPointer );
+
+							if ( IViewportInteractableInterface* ActorInteractable = Cast<IViewportInteractableInterface>( Actor ) )
+							{
+								if ( Action->ActionType == ViewportWorldActionTypes::SelectAndMove )
+								{
+									bHandled = true;
+
+									bool bResultedInInteractableDrag = false;
+									ActorInteractable->OnPressed( this, HitResult, bResultedInInteractableDrag );
+
+									if ( bResultedInInteractableDrag )
+									{
+										InteractorData.DraggingMode = EViewportInteractionDraggingMode::Interactable;
+										WorldInteraction->SetDraggedInteractable( ActorInteractable );
+									}
+								}
+							}
+							else
+							{
+								bHandled = true;
+
+								// Is the other hand already dragging this stuff?
+								if ( OtherInteractorData.DraggingMode == EViewportInteractionDraggingMode::ActorsWithGizmo ||
+									OtherInteractorData.DraggingMode == EViewportInteractionDraggingMode::ActorsFreely )
+								{
+									// Only if they clicked on one of the objects we're moving
+									if ( WorldInteraction->IsTransformingActor( Actor ) )
+									{
+										// If we were dragging with a gizmo, we'll need to stop doing that and instead drag freely.
+										if ( OtherInteractorData.DraggingMode == EViewportInteractionDraggingMode::ActorsWithGizmo )
+										{
+											WorldInteraction->StopDragging( this );
+
+											FViewportActionKeyInput OtherInteractorAction( ViewportWorldActionTypes::SelectAndMove );
+											const bool bIsPlacingActors = false;
+											WorldInteraction->StartDraggingActors( OtherInteractor, OtherInteractorAction, HitResult.GetComponent(), OtherInteractorData.HoverLocation, bIsPlacingActors ); //@todo ViewportInteraction
+										}
+
+										InteractorData.DraggingMode = InteractorData.LastDraggingMode = EViewportInteractionDraggingMode::AssistingDrag;
+
+										InteractorData.bIsFirstDragUpdate = true;
+										InteractorData.bWasAssistingDrag = true;
+										InteractorData.DragRayLength = ( HitResult.ImpactPoint - LaserPointerStart ).Size();
+										InteractorData.LastLaserPointerStart = LaserPointerStart;
+										InteractorData.LastDragToLocation = HitResult.ImpactPoint;
+										InteractorData.LaserPointerImpactAtDragStart = HitResult.ImpactPoint;
+										InteractorData.DragTranslationVelocity = FVector::ZeroVector;
+										InteractorData.DragRayLengthVelocity = 0.0f;
+										InteractorData.DraggingTransformGizmoComponent = nullptr;
+										InteractorData.TransformGizmoInteractionType = ETransformGizmoInteractionType::None;
+										InteractorData.bIsDrivingVelocityOfSimulatedTransformables = false;
+
+										InteractorData.GizmoStartTransform = OtherInteractorData.GizmoStartTransform;
+										InteractorData.GizmoStartLocalBounds = OtherInteractorData.GizmoStartLocalBounds;
+										InteractorData.GizmoSpaceFirstDragUpdateOffsetAlongAxis = FVector::ZeroVector;	// Will be determined on first update
+										InteractorData.GizmoSpaceDragDeltaFromStartOffset = FVector::ZeroVector;	// Set every frame while dragging
+
+										WorldInteraction->SetDraggedSinceLastSelection( true );
+										WorldInteraction->SetLastDragGizmoStartTransform( InteractorData.GizmoStartTransform );
+
+										Action->bIsInputCaptured = true; 
+
+										// Play a haptic effect when objects are picked up
+										PlayHapticEffect( VI::DragHapticFeedbackStrength->GetFloat() ); //@todo ViewportInteraction
+									}
+									else
+									{
+										// @todo vreditor: We don't currently support selection/dragging separate objects with either hand
+									}
+								}
+								else if ( OtherInteractorData.DraggingMode == EViewportInteractionDraggingMode::ActorsWithGizmo )
+								{
+									// We don't support dragging objects with the gizmo using two hands.  Just absorb it.
+								}
+								else if ( OtherInteractorData.DraggingMode == EViewportInteractionDraggingMode::ActorsAtLaserImpact )
+								{
+									// Doesn't work with two hands.  Just absorb it.
+								}
+								else
+								{
+									// Default to replacing our selection with whatever the user clicked on
+									enum ESelectionModification
+									{
+										Replace,
+										AddTo,
+										Toggle
+									} SelectionModification = ESelectionModification::Replace;
+
+									// If the other hand is holding down the button after selecting an object, we'll allow this hand to toggle selection
+									// of additional objects (multi select)
+									{
+										if ( OtherInteractorData.ClickingOnComponent.IsValid() &&		// Other hand is clicking on something
+											Actor != WorldInteraction->GetTransformGizmoActor() &&					// Not clicking on a gizmo
+											OtherInteractorData.ClickingOnComponent.Get()->GetOwner() != Actor )	// Not clicking on same actor
+										{
+											// OK, the other hand is holding down the "select and move" button.
+											SelectionModification = ESelectionModification::Toggle;
+										}
+									}
+
+									// Only start dragging if the actor was already selected (and it's a full press), or if we clicked on a gizmo.  It's too easy to 
+									// accidentally move actors around the scene when you only want to select them otherwise.
+									bool bShouldDragSelected =
+										( ( Actor->IsSelected() && Action->ActionType == ViewportWorldActionTypes::SelectAndMove ) || Actor == WorldInteraction->GetTransformGizmoActor() ) &&
+										( SelectionModification == ESelectionModification::Replace );
+
+									bool bSelectionChanged = false;
+									if ( Actor != WorldInteraction->GetTransformGizmoActor() )  // Don't change any actor selection state if the user clicked on a gizmo
+									{
+										const bool bWasSelected = Actor->IsSelected();
+
+										// Clicked on a normal actor.  So update selection!
+										if ( SelectionModification == ESelectionModification::Replace || SelectionModification == ESelectionModification::AddTo )
+										{
+											if ( !bWasSelected )
+											{
+												if ( SelectionModification == ESelectionModification::Replace )
+												{
+													GEditor->SelectNone( true, true );
+												}
+
+												GEditor->SelectActor( Actor, true, true );
+												bSelectionChanged = true;
+											}
+										}
+										else if ( SelectionModification == ESelectionModification::Toggle )
+										{
+											GEditor->SelectActor( Actor, !Actor->IsSelected(), true );
+											bSelectionChanged = true;
+										}
+
+										// If we've selected an actor with a light press, don't prevent the user from moving the actor if they
+										// press down all the way.  We only want locking when interacting with gizmos
+										InteractorData.bAllowTriggerLightPressLocking = false;
+
+										// If the user did a full press to select an actor that wasn't selected, allow them to drag it right away
+										const bool bOtherHandTryingToDrag =
+											OtherInteractorData.ClickingOnComponent.IsValid() &&
+											OtherInteractorData.ClickingOnComponent.Get()->GetOwner()->IsSelected() &&
+											OtherInteractorData.ClickingOnComponent.Get()->GetOwner() == HitResult.GetComponent()->GetOwner();	// Trying to drag same actor
+										if ( ( !bWasSelected || bOtherHandTryingToDrag ) &&
+											Actor->IsSelected() &&
+											( Action->ActionType == ViewportWorldActionTypes::SelectAndMove || bOtherHandTryingToDrag ) )
+										{
+											bShouldDragSelected = true;
+										}
+									}
+
+									if ( bShouldDragSelected && InteractorData.DraggingMode != EViewportInteractionDraggingMode::Interactable )
+									{
+										const bool bIsPlacingActors = false;
+										WorldInteraction->StartDraggingActors( this, *Action, HitResult.GetComponent(), HitResult.ImpactPoint, bIsPlacingActors ); //@todo ViewportInteraction
+									}
+									else if ( bSelectionChanged )
+									{
+										// User didn't drag but did select something, so play a haptic effect.
+										PlayHapticEffect( VI::SelectionHapticFeedbackStrength->GetFloat() );
+									}
+
+									bHandled;
+								}
+							}
+						}
+					}
+				}
+				else if ( Event == IE_Released )
+				{
+					if ( InteractorData.ClickingOnComponent.IsValid() )
+					{
+						bHandled = true;
+						InteractorData.ClickingOnComponent = nullptr;
+					}
+
+					const bool bIsDraggingWorld = InteractorData.DraggingMode == EViewportInteractionDraggingMode::World;
+					const bool bIsDraggingWorldWithTwoHands =
+						( InteractorData.DraggingMode == EViewportInteractionDraggingMode::World && GetOtherInteractor()->InteractorData.DraggingMode == EViewportInteractionDraggingMode::AssistingDrag ) ||
+						( GetOtherInteractor()->InteractorData.DraggingMode == EViewportInteractionDraggingMode::World && InteractorData.DraggingMode == EViewportInteractionDraggingMode::AssistingDrag );
+
+					// Don't allow the trigger to cancel our drag on release if we're dragging the world. 
+					if ( InteractorData.DraggingMode != EViewportInteractionDraggingMode::Nothing &&
+						!bIsDraggingWorld &&
+						!bIsDraggingWorldWithTwoHands )
+					{
+						WorldInteraction->StopDragging( this );
+					}
+
+					bHandled = true;
+				}
+			}
+			// World Movement
+			else if ( Action->ActionType == ViewportWorldActionTypes::WorldMovement )
+			{
+				if ( Event == IE_Pressed )
+				{
+					// Is our other hand already dragging the world around?
+					if ( OtherInteractor && OtherInteractor->GetInteractorData().DraggingMode == EViewportInteractionDraggingMode::World )
+					{
+						InteractorData.DraggingMode = InteractorData.LastDraggingMode = EViewportInteractionDraggingMode::AssistingDrag;
+						InteractorData.bWasAssistingDrag = true;
+					}
+					else
+					{
+						// Start dragging the world
+						InteractorData.DraggingMode = InteractorData.LastDraggingMode = EViewportInteractionDraggingMode::World;
+						InteractorData.bWasAssistingDrag = false;
+
+						// Starting a new drag, so make sure the other hand doesn't think it's assisting us
+						OtherInteractor->GetInteractorData().bWasAssistingDrag = false;
+
+						// Stop any interia from the other hand's previous movements -- we've grabbed the world and it needs to stick!
+						OtherInteractor->GetInteractorData().DragTranslationVelocity = FVector::ZeroVector;
+					}
+
+					InteractorData.bIsFirstDragUpdate = true;
+					InteractorData.DragRayLength = 0.0f;
+					InteractorData.LastLaserPointerStart = InteractorData.Transform.GetLocation();
+					InteractorData.LastDragToLocation = InteractorData.Transform.GetLocation();
+					InteractorData.DragTranslationVelocity = FVector::ZeroVector;
+					InteractorData.DragRayLengthVelocity = 0.0f;
+					InteractorData.bIsDrivingVelocityOfSimulatedTransformables = false;
+
+					// We won't use gizmo features for world movement
+					InteractorData.DraggingTransformGizmoComponent = nullptr;
+					InteractorData.TransformGizmoInteractionType = ETransformGizmoInteractionType::None;
+					InteractorData.OptionalHandlePlacement.Reset();
+					InteractorData.GizmoStartTransform = FTransform::Identity;
+					InteractorData.GizmoStartLocalBounds = FBox( 0 );
+					InteractorData.GizmoSpaceFirstDragUpdateOffsetAlongAxis = FVector::ZeroVector;
+					InteractorData.GizmoSpaceDragDeltaFromStartOffset = FVector::ZeroVector;
+
+					Action->bIsInputCaptured = true;
+				}
+				else if ( Event == IE_Released )
+				{
+					WorldInteraction->StopDragging( this );
+				}
+			}
+			else if ( Action->ActionType == ViewportWorldActionTypes::Delete )
+			{
+				if ( Event == IE_Pressed )
+				{
+					WorldInteraction->DeleteSelectedObjects();
+				}
+				bHandled = true;
+			}
+			else if ( Action->ActionType == ViewportWorldActionTypes::Redo )
+			{
+				if ( Event == IE_Pressed || Event == IE_Repeat )
+				{
+					WorldInteraction->Redo();
+				}
+				bHandled = true;
+			}
+			else if ( Action->ActionType == ViewportWorldActionTypes::Undo )
+			{
+				if ( Event == IE_Pressed || Event == IE_Repeat )
+				{
+					WorldInteraction->Undo();
+				}
+				bHandled = true;
+			}
+		}
+
+		if ( !bHandled )
+		{
+			// If "select and move" was pressed but not handled, go ahead and deselect everything
+			if ( Action->ActionType == ViewportWorldActionTypes::SelectAndMove && Event == IE_Pressed )
+			{
+				WorldInteraction->Deselect();
+			}
+		}
+	}
+
+	return bHandled;
+}
+
+bool UViewportInteractor::HandleInputAxis( const FKey Key, const float Delta, const float DeltaTime )
+{
+	bool bHandled = false;
+	FViewportActionKeyInput* KnownAction = KeyToActionMap.Find( Key );
+	if ( KnownAction != nullptr )	// Ignore key repeats
+	{
+		FViewportActionKeyInput Action( KnownAction->ActionType );
+
+		// Give the derived classes a chance to update according to the input
+		HandleInputAxis( Action, Key, Delta, DeltaTime, bHandled );
+	}
+
+	return bHandled;
+}
+
+FTransform UViewportInteractor::GetTransform() const
+{
+	return InteractorData.Transform;
+}
+
+EViewportInteractionDraggingMode UViewportInteractor::GetDraggingMode() const
+{
+	return InteractorData.DraggingMode;
+}
+
+EViewportInteractionDraggingMode UViewportInteractor::GetLastDraggingMode() const
+{
+	return InteractorData.LastDraggingMode;
+}
+
+FVector UViewportInteractor::GetDragTranslationVelocity() const
+{
+	return InteractorData.DragTranslationVelocity;
+}
+
+bool UViewportInteractor::GetLaserPointer( FVector& LaserPointerStart, FVector& LaserPointerEnd, const bool bEvenIfBlocked, const float LaserLengthOverride )
+{
+	// If we're currently grabbing the world with both hands, then the laser pointer is not available
+	if ( /*bHaveMotionController &&*/ //@todo ViewportInteraction
+		!( InteractorData.DraggingMode == EViewportInteractionDraggingMode::World && GetOtherInteractor()->GetInteractorData().DraggingMode == EViewportInteractionDraggingMode::AssistingDrag ) &&
+		!( InteractorData.DraggingMode == EViewportInteractionDraggingMode::AssistingDrag && GetOtherInteractor()->GetInteractorData().DraggingMode == EViewportInteractionDraggingMode::World ) )
+	{
+		// If we have UI attached to us, don't allow a laser pointer
+		if ( bEvenIfBlocked || !GetIsLaserBlocked() )
+		{
+			FTransform HandTransform;
+			FVector HandForwardVector;
+			if ( GetTransformAndForwardVector( HandTransform, HandForwardVector ) )
+			{
+				LaserPointerStart = HandTransform.GetLocation();
+
+				float LaserLength = LaserLengthOverride == 0.0f ? GetLaserPointerMaxLength() : LaserLengthOverride;
+				LaserPointerEnd = LaserPointerStart + HandForwardVector * LaserLength;
+
+				return true;
+			}
+		}
+	}
+	return false;
+}
+
+float UViewportInteractor::GetLaserPointerMaxLength() const
+{
+	return VI::LaserPointerMaxLength->GetFloat() * WorldInteraction->GetWorldScaleFactor();
+}
+
+FHitResult UViewportInteractor::GetHitResultFromLaserPointer( TArray<AActor*>* OptionalListOfIgnoredActors /*= nullptr*/, const bool bIgnoreGizmos /*= false*/, 
+	TArray<UClass*>* ObjectsInFrontOfGizmo /*= nullptr */, const bool bEvenIfBlocked /*= false */, const float LaserLengthOverride /*= 0.0f */ )
+{
+	FHitResult BestHitResult;
+
+	FVector LaserPointerStart, LaserPointerEnd;
+	if ( GetLaserPointer( LaserPointerStart, LaserPointerEnd, bEvenIfBlocked, LaserLengthOverride ) )
+	{
+		// Twice twice.  Once for editor gizmos which are "on top" and always take precedence, then a second time
+		// for all of the scene objects
+		for ( int32 PassIndex = bIgnoreGizmos ? 1 : 0; PassIndex < 2; ++PassIndex )
+		{
+			const bool bOnlyEditorGizmos = ( PassIndex == 0 );
+
+			const bool bTraceComplex = true;
+			FCollisionQueryParams TraceParams( NAME_None, bTraceComplex, nullptr );
+
+			const FCollisionResponseParams& ResponseParam = FCollisionResponseParams::DefaultResponseParam;
+
+			const ECollisionChannel CollisionChannel = bOnlyEditorGizmos ? COLLISION_GIZMO : ECC_Visibility;
+
+			if ( OptionalListOfIgnoredActors != nullptr )
+			{
+				TraceParams.AddIgnoredActors( *OptionalListOfIgnoredActors );
+			}
+
+			bool bHit = false;
+			FHitResult HitResult;
+			if ( bOnlyEditorGizmos )
+			{
+				bHit = WorldInteraction->GetViewportWorld()->LineTraceSingleByChannel( HitResult, LaserPointerStart, LaserPointerEnd, CollisionChannel, TraceParams, ResponseParam );
+			}
+			else
+			{
+				FCollisionObjectQueryParams EverythingButGizmos( FCollisionObjectQueryParams::AllObjects );
+				EverythingButGizmos.RemoveObjectTypesToQuery( COLLISION_GIZMO );
+				bHit = WorldInteraction->GetViewportWorld()->LineTraceSingleByObjectType( HitResult, LaserPointerStart, LaserPointerEnd, EverythingButGizmos, TraceParams );
+			}
+
+			if ( bHit )
+			{
+				bool bHitResultIsPriorityType = false;
+				if ( !bOnlyEditorGizmos && ObjectsInFrontOfGizmo )
+				{
+					for ( UClass* CurrentClass : *ObjectsInFrontOfGizmo )
+					{
+						bool bClassHasPriority = false;
+						bClassHasPriority = HitResult.GetActor()->IsA( CurrentClass );
+
+						if ( bClassHasPriority )
+						{
+							bHitResultIsPriorityType = bClassHasPriority;
+							break;
+						}
+					}
+				}
+
+				if ( BestHitResult.GetActor() == nullptr || bHitResultIsPriorityType )
+				{
+					BestHitResult = HitResult;
+				}
+			}
+		}
+	}
+
+	return BestHitResult;
+}
+
+bool UViewportInteractor::GetTransformAndForwardVector( FTransform& OutHandTransform, FVector& OutForwardVector )
+{
+	return false;
+}
+
+FVector UViewportInteractor::GetHoverLocation() const 
+{
+	return InteractorData.HoverLocation;
+}
+
+bool UViewportInteractor::IsHovering() const
+{
+	return InteractorData.bIsHovering;
+}
+
+void UViewportInteractor::SetDraggingMode( const EViewportInteractionDraggingMode NewDraggingMode )
+{
+	InteractorData.DraggingMode = NewDraggingMode;
+}
+
+FViewportActionKeyInput* UViewportInteractor::GetActionWithName( const FName InActionName )
+{
+	FViewportActionKeyInput* ResultedAction = nullptr;
+	for ( auto It = KeyToActionMap.CreateIterator(); It; ++It )
+	{
+		if ( It.Value().ActionType == InActionName )
+		{
+			ResultedAction = &(It->Value);
+			break;
+		}
+	}
+	return ResultedAction;
+}
+
+float UViewportInteractor::GetDragHapticFeedbackStrength() const
+{
+	return VI::DragHapticFeedbackStrength->GetFloat();
+}
+
+void UViewportInteractor::SetAllowTriggerLightPressLocking( const bool bInAllow )
+{
+	InteractorData.bAllowTriggerLightPressLocking = bInAllow;
+}
+
+bool UViewportInteractor::AllowTriggerLightPressLocking() const
+{
+	return InteractorData.bAllowTriggerLightPressLocking;
+}
\ No newline at end of file
diff --git a/Engine/Source/Editor/ViewportInteraction/ViewportInteractorData.h b/Engine/Source/Editor/ViewportInteraction/ViewportInteractorData.h
new file mode 100644
--- /dev/null
+++ b/Engine/Source/Editor/ViewportInteraction/ViewportInteractorData.h
@@ -0,0 +1,171 @@
+// Copyright 1998-2015 Epic Games, Inc. All Rights Reserved.
+
+#pragma once
+
+#include "ViewportInteractionTypes.h"
+#include "ViewportInteractionTypes.h"
+
+/** Represents a single virtual hand */
+struct FViewportInteractorData
+{
+	//
+	// Positional data
+	//
+
+	/** Your hand in the virtual world in world space, usually driven by VR motion controllers */
+	FTransform Transform;
+
+	/** Your hand transform, in the local tracking space */
+	FTransform RoomSpaceTransform;
+
+	/** Hand transform in world space from the previous frame */
+	FTransform LastTransform;
+
+	/** Room space hand transform from the previous frame */
+	FTransform LastRoomSpaceTransform;
+
+	//
+	// Hover feedback
+	//
+
+	/** Lights up when our laser pointer is close enough to something to click */
+	bool bIsHovering;
+
+	/** True if we're hovering over UI right now.  When hovering over UI, we don't bother drawing a see-through laser pointer */
+	bool bIsHoveringOverUI; //@todo: ViewportInteraction: UI should not be in this module
+
+	/** The widget component we are currently hovering over */
+	class UWidgetComponent* HoveringOverWidgetComponent; //@todo: ViewportInteraction: UI should not be in this module
+
+	/** Position the laser pointer impacted an interactive object at (UI, meshes, etc.) */
+	FVector HoverLocation;
+
+	/** The current component hovered by the laser pointer of this hand */
+	TWeakObjectPtr<class UActorComponent> HoveredActorComponent;
+
+	//
+	// General input
+	//
+
+	/** True if we're currently holding the 'SelectAndMove' button down after clicking on an actor */
+	TWeakObjectPtr<UActorComponent> ClickingOnComponent;
+
+	/** Last real time that we pressed the 'SelectAndMove' button on UI.  This is used to detect double-clicks. */
+	double LastClickPressTime;
+
+
+	//
+	// Object/world movement
+	// 
+
+	/** What we're currently dragging with this hand, if anything */
+	EViewportInteractionDraggingMode DraggingMode;
+
+	/** What we were doing last.  Used for inertial movement. */
+	EViewportInteractionDraggingMode LastDraggingMode;
+
+	/** True if this is the first update since we started dragging */
+	bool bIsFirstDragUpdate;
+
+	/** True if we were assisting the other hand's drag the last time we did anything.  This is used for inertial movement */
+	bool bWasAssistingDrag;
+
+	/** Length of the ray that's dragging */
+	float DragRayLength;
+
+	/** Laser pointer start location the last frame */
+	FVector LastLaserPointerStart;
+
+	/** Location that we dragged to last frame (end point of the ray) */
+	FVector LastDragToLocation;
+
+	/** Laser pointer impact location at the drag start */
+	FVector LaserPointerImpactAtDragStart;
+
+	/** How fast to move selected objects every frame for inertial translation */
+	FVector DragTranslationVelocity;
+
+	/** How fast to adjust ray length every frame for inertial ray length changes */
+	float DragRayLengthVelocity;
+
+	/** While dragging, true if we're dragging at least one simulated object that we're driving the velocities of.  When this
+	    is true, our default inertia system is disabled and we rely on the physics engine to take care of inertia */
+	bool bIsDrivingVelocityOfSimulatedTransformables;
+
+
+	//
+	// Transform gizmo interaction
+	//
+
+	/** Where the gizmo was placed at the beginning of the current interaction */
+	FTransform GizmoStartTransform;
+
+	/** Our gizmo bounds at the start of the interaction, in actor local space. */
+	FBox GizmoStartLocalBounds;
+
+	/** For a single axis drag, this is the cached local offset where the laser pointer ray intersected the axis line on the first frame of the drag */
+	FVector GizmoSpaceFirstDragUpdateOffsetAlongAxis;
+
+	/** When dragging with an axis/plane constraint applied, this is the difference between the actual "delta from start" and the constrained "delta from start".
+		This is used when the user releases the object and inertia kicks in */
+	FVector GizmoSpaceDragDeltaFromStartOffset;
+
+	/** The gizmo interaction we're doing with this hand */
+	ETransformGizmoInteractionType TransformGizmoInteractionType;
+
+	/** Which handle on the gizmo we're interacting with, if any */
+	TOptional<FTransformGizmoHandlePlacement> OptionalHandlePlacement;
+
+	/** The gizmo component we're dragging right now */
+	TWeakObjectPtr<class USceneComponent> DraggingTransformGizmoComponent;
+
+	/** Gizmo component that we're hovering over, or nullptr if not hovering over any */
+	TWeakObjectPtr<class USceneComponent> HoveringOverTransformGizmoComponent;
+
+	/** Gizmo handle that we hovered over last (used only for avoiding spamming of hover haptics!) */
+	TWeakObjectPtr<class USceneComponent> HoverHapticCheckLastHoveredGizmoComponent;
+
+	//@todo ViewportInteraction: Light press should not have to be here, ViewportInteractor needs it
+	/** True if we allow locking of the lightly pressed state for the current event.  This can be useful to
+	turn off in cases where you always want a full press to override the light press, even if it was
+	a very slow press */
+	bool bAllowTriggerLightPressLocking;
+
+	/** Default constructor for FVirtualHand that initializes safe defaults */
+	FViewportInteractorData()
+	{
+		bIsHovering = false;
+		bIsHoveringOverUI = false;
+		HoveringOverWidgetComponent = nullptr;
+		HoverLocation = FVector::ZeroVector;
+		HoveredActorComponent = nullptr;
+
+		ClickingOnComponent = nullptr;
+
+		DraggingMode = EViewportInteractionDraggingMode::Nothing;
+		LastDraggingMode = EViewportInteractionDraggingMode::Nothing;
+		bWasAssistingDrag = false;
+		bIsFirstDragUpdate = false;
+		DragRayLength = 0.0f;
+		LastLaserPointerStart = FVector::ZeroVector;
+		LastDragToLocation = FVector::ZeroVector;
+		LaserPointerImpactAtDragStart = FVector::ZeroVector;
+		DragTranslationVelocity = FVector::ZeroVector;
+		DragRayLengthVelocity = 0.0f;
+		bIsDrivingVelocityOfSimulatedTransformables = false;
+
+		GizmoStartTransform = FTransform::Identity;
+		GizmoStartLocalBounds = FBox( 0 );
+		GizmoSpaceFirstDragUpdateOffsetAlongAxis = FVector::ZeroVector;
+		GizmoSpaceDragDeltaFromStartOffset = FVector::ZeroVector;
+		TransformGizmoInteractionType = ETransformGizmoInteractionType::None;
+		OptionalHandlePlacement.Reset();
+		DraggingTransformGizmoComponent = nullptr;
+		HoveringOverTransformGizmoComponent = nullptr;
+		HoverHapticCheckLastHoveredGizmoComponent = nullptr;
+
+		bAllowTriggerLightPressLocking = true;
+	}
+};
+
+
diff --git a/Engine/Source/Editor/ViewportInteraction/ViewportInteractorMotionController.cpp b/Engine/Source/Editor/ViewportInteraction/ViewportInteractorMotionController.cpp
new file mode 100644
--- /dev/null
+++ b/Engine/Source/Editor/ViewportInteraction/ViewportInteractorMotionController.cpp
@@ -0,0 +1,350 @@
+// Copyright 1998-2015 Epic Games, Inc. All Rights Reserved.
+
+#include "ViewportInteractionModule.h"
+#include "ViewportInteractorMotionController.h"
+#include "ViewportInteractorData.h"
+#include "ViewportWorldInteraction.h"
+
+#include "IMotionController.h"
+#include "IHeadMountedDisplay.h"
+#include "MotionControllerComponent.h"
+#include "Features/IModularFeatures.h"
+
+namespace VI
+{
+	//Laser
+	static FAutoConsoleVariable OculusLaserPointerRotationOffset( TEXT( "VI.OculusLaserPointerRotationOffset" ), 0.0f, TEXT( "How much to rotate the laser pointer (pitch) relative to the forward vector of the controller (Oculus)" ) );
+	static FAutoConsoleVariable ViveLaserPointerRotationOffset( TEXT( "VI.ViveLaserPointerRotationOffset" ), /* -57.8f */ 0.0f, TEXT( "How much to rotate the laser pointer (pitch) relative to the forward vector of the controller (Vive)" ) );
+	static FAutoConsoleVariable OculusLaserPointerStartOffset( TEXT( "VI.OculusLaserPointerStartOffset" ), 2.8f, TEXT( "How far to offset the start of the laser pointer to avoid overlapping the hand mesh geometry (Oculus)" ) );
+	static FAutoConsoleVariable ViveLaserPointerStartOffset( TEXT( "VI.ViveLaserPointerStartOffset" ), 1.25f /* 8.5f */, TEXT( "How far to offset the start of the laser pointer to avoid overlapping the hand mesh geometry (Vive)" ) );
+	
+	//Trigger
+	static FAutoConsoleVariable TriggerLightlyPressedThreshold( TEXT( "VI.TriggerLightlyPressedThreshold" ), 0.03f, TEXT( "Minimum trigger threshold before we consider the trigger at least 'lightly pressed'" ) );
+	static FAutoConsoleVariable TriggerDeadZone( TEXT( "VI.TriggerDeadZone" ), 0.01f, TEXT( "Trigger dead zone.  The trigger must be fully released before we'll trigger a new 'light press'" ) );
+	static FAutoConsoleVariable TriggerFullyPressedThreshold( TEXT( "VI.TriggerFullyPressedThreshold" ), 0.95f, TEXT( "Minimum trigger threshold before we consider the trigger 'fully pressed'" ) );
+	static FAutoConsoleVariable TriggerFullyPressedReleaseThreshold( TEXT( "VI.TriggerFullyPressedReleaseThreshold" ), 0.8f, TEXT( "After fully pressing the trigger, if the axis falls below this threshold we no longer consider it fully pressed" ) );
+
+	// Haptic feedback
+	static FAutoConsoleVariable SleepForRiftHaptics( TEXT( "VI.SleepForRiftHaptics" ), 1, TEXT( "When enabled, we'll sleep the game thread mid-frame to wait for haptic effects to finish.  This can be devasting to performance!" ) );
+	static FAutoConsoleVariable MinHapticTimeForRift( TEXT( "VREd.MinHapticTimeForRift" ), 0.005f, TEXT( "How long to play haptic effects on the Rift" ) );
+}
+
+const FName UViewportInteractorMotionController::TrackpadPositionX = FName( "TrackpadPositionX" );
+const FName UViewportInteractorMotionController::TrackpadPositionY = FName( "TrackpadPositionY" );
+const FName UViewportInteractorMotionController::TriggerAxis = FName( "TriggerAxis" );
+const FName UViewportInteractorMotionController::MotionController_Left_FullyPressedTriggerAxis = FName( "MotionController_Left_FullyPressedTriggerAxis" );
+const FName UViewportInteractorMotionController::MotionController_Right_FullyPressedTriggerAxis = FName( "MotionController_Right_FullyPressedTriggerAxis" );
+const FName UViewportInteractorMotionController::MotionController_Left_LightlyPressedTriggerAxis = FName( "MotionController_Left_LightlyPressedTriggerAxis" );
+const FName UViewportInteractorMotionController::MotionController_Right_LightlyPressedTriggerAxis = FName( "MotionController_Right_LightlyPressedTriggerAxis" );
+
+UViewportInteractorMotionController::UViewportInteractorMotionController( const FObjectInitializer& Initializer ) :
+	UViewportInteractor( Initializer ),
+	MotionControllerComponent( nullptr ),
+	HandMeshComponent( nullptr ),
+	LaserPointerMeshComponent( nullptr ),
+	LaserPointerMID( nullptr ),
+	TranslucentLaserPointerMID( nullptr ),
+	HoverMeshComponent( nullptr ),
+	HoverPointLightComponent( nullptr ),
+	HandMeshMID( nullptr ),
+	ControllerHandSide( EControllerHand::Pad ),
+	bHaveMotionController( false ),
+	bIsTriggerFullyPressed( false ),
+	bIsTriggerAtLeastLightlyPressed( false ),
+	RealTimeTriggerWasLightlyPressed( 0.0f ),
+	bHasTriggerBeenReleasedSinceLastPress( true ),
+	SelectAndMoveTriggerValue( 0.0f )
+{
+}
+
+UViewportInteractorMotionController::~UViewportInteractorMotionController()
+{
+	Shutdown();
+}
+
+void UViewportInteractorMotionController::Shutdown()
+{
+	Super::Shutdown();
+
+	MotionControllerComponent = nullptr;
+	HandMeshComponent = nullptr;
+	LaserPointerMeshComponent = nullptr;
+	LaserPointerMID = nullptr;
+	TranslucentLaserPointerMID = nullptr;
+	HoverMeshComponent = nullptr;
+	HoverPointLightComponent = nullptr;
+	HandMeshComponent = nullptr;
+}
+
+void UViewportInteractorMotionController::Tick( const float DeltaTime )
+{
+	Super::Tick( DeltaTime );
+
+	{
+		const float WorldScaleFactor = WorldInteraction->GetWorldScaleFactor();
+
+		// @todo vreditor: Manually ticking motion controller components
+		MotionControllerComponent->TickComponent( DeltaTime, ELevelTick::LEVELTICK_PauseTick, nullptr );
+
+		// The hands need to stay the same size relative to our tracking space, so we inverse compensate for world to meters scale here
+		// NOTE: We don't need to set the hand mesh location and rotation, as the MotionControllerComponent does that itself
+		if ( ControllerHandSide == EControllerHand::Right &&
+			GetHMDDeviceType() == EHMDDeviceType::DT_OculusRift )	// Oculus has asymmetrical controllers, so we mirror the mesh horizontally
+		{
+			HandMeshComponent->SetRelativeScale3D( FVector( WorldScaleFactor, -WorldScaleFactor, WorldScaleFactor ) );
+		}
+		else
+		{
+			HandMeshComponent->SetRelativeScale3D( FVector( WorldScaleFactor ) );
+		}
+	}
+}
+
+EHMDDeviceType::Type UViewportInteractorMotionController::GetHMDDeviceType() const
+{
+	return GEngine->HMDDevice.IsValid() ? GEngine->HMDDevice->GetHMDDeviceType() : EHMDDeviceType::DT_SteamVR; //@todo: ViewportInteraction, assumption that it's steamvr ??
+}
+
+void UViewportInteractorMotionController::HandleInputAxis( FViewportActionKeyInput& Action, const FKey Key, const float Delta, const float DeltaTime, bool& bOutWasHandled )
+{
+	if ( Action.ActionType == TriggerAxis )
+	{
+		// Synthesize "lightly pressed" events for the trigger
+		{
+			// Store latest trigger value amount
+			SelectAndMoveTriggerValue = Delta;
+
+			FViewportActionKeyInput* OptionalAction = GetActionWithName( ViewportWorldActionTypes::SelectAndMove );
+			const bool bIsFullPressAlreadyCapturing = OptionalAction && OptionalAction->bIsInputCaptured;
+				
+			if ( !bIsFullPressAlreadyCapturing &&		// Don't fire if we're already capturing input for a full press
+				!bIsTriggerAtLeastLightlyPressed &&	// Don't fire unless we are already pressed
+				bHasTriggerBeenReleasedSinceLastPress &&	// Only if we've been fully released since the last time we fired
+				Delta >= VI::TriggerLightlyPressedThreshold->GetFloat() )
+			{
+				bIsTriggerAtLeastLightlyPressed = true;
+				SetAllowTriggerLightPressLocking( true );
+				RealTimeTriggerWasLightlyPressed = FPlatformTime::Seconds();
+				bHasTriggerBeenReleasedSinceLastPress = false;
+
+				// Synthesize an input key for this light press
+				const EInputEvent InputEvent = IE_Pressed;
+				const bool bWasLightPressHandled = UViewportInteractor::HandleInputKey( ControllerHandSide == EControllerHand::Left ? MotionController_Left_LightlyPressedTriggerAxis : MotionController_Right_LightlyPressedTriggerAxis, InputEvent );
+			}
+			else if ( bIsTriggerAtLeastLightlyPressed && Delta < VI::TriggerLightlyPressedThreshold->GetFloat() )
+			{
+				bIsTriggerAtLeastLightlyPressed = false;
+
+				// Synthesize an input key for this light press
+				const EInputEvent InputEvent = IE_Released;
+				const bool bWasLightReleaseHandled = UViewportInteractor::HandleInputKey( ControllerHandSide == EControllerHand::Left ? MotionController_Left_LightlyPressedTriggerAxis : MotionController_Right_LightlyPressedTriggerAxis, InputEvent );
+			}
+		}
+
+		if ( Delta < VI::TriggerDeadZone->GetFloat() )
+		{
+			bHasTriggerBeenReleasedSinceLastPress = true;
+		}
+
+		// Synthesize "fully pressed" events for the trigger
+		{
+			if ( !bIsTriggerFullyPressed &&	// Don't fire unless we are already pressed
+				Delta >= VI::TriggerFullyPressedThreshold->GetFloat() )
+			{
+				bIsTriggerFullyPressed = true;
+
+				// Synthesize an input key for this full press
+				const EInputEvent InputEvent = IE_Pressed;
+				UViewportInteractor::HandleInputKey( ControllerHandSide == EControllerHand::Left ? MotionController_Left_FullyPressedTriggerAxis : MotionController_Right_FullyPressedTriggerAxis, InputEvent );
+			}
+			else if ( bIsTriggerFullyPressed && Delta < VI::TriggerFullyPressedReleaseThreshold->GetFloat() )
+			{
+				bIsTriggerFullyPressed = false;
+
+				// Synthesize an input key for this full press
+				const EInputEvent InputEvent = IE_Released;
+				UViewportInteractor::HandleInputKey( ControllerHandSide == EControllerHand::Left ? MotionController_Left_FullyPressedTriggerAxis : MotionController_Right_FullyPressedTriggerAxis, InputEvent );
+			}
+		}
+	}
+}
+
+void UViewportInteractorMotionController::PollInput()
+{
+	bHaveMotionController = false;
+	InteractorData.LastTransform = InteractorData.Transform;
+	InteractorData.LastRoomSpaceTransform = InteractorData.RoomSpaceTransform;
+
+	//@todo ViewportInteraction: VREditorMode::PollInputFromMotionControllers mouse implementation 
+
+	// Generic motion controllers
+	if ( !bHaveMotionController )
+	{
+		TArray<IMotionController*> MotionControllers = IModularFeatures::Get().GetModularFeatureImplementations<IMotionController>( IMotionController::GetModularFeatureName() );
+		for ( auto MotionController : MotionControllers )
+		{
+			if ( MotionController != nullptr && !bHaveMotionController )
+			{
+				FVector Location = FVector::ZeroVector;
+				FRotator Rotation = FRotator::ZeroRotator;
+
+				if ( MotionController->GetControllerOrientationAndPosition( WorldInteraction->GetMotionControllerID(), ControllerHandSide, /* Out */ Rotation, /* Out */ Location ) )
+				{
+					bHaveMotionController = true;
+					InteractorData.RoomSpaceTransform = FTransform( Rotation.Quaternion(), Location, FVector( 1.0f ) );
+					InteractorData.Transform = InteractorData.RoomSpaceTransform * WorldInteraction->GetRoomTransform();
+				}
+			}
+		}
+	}
+}
+
+bool UViewportInteractorMotionController::IsTriggerAtLeastLightlyPressed() const
+{
+	return bIsTriggerAtLeastLightlyPressed;
+}
+
+double UViewportInteractorMotionController::GetRealTimeTriggerWasLightlyPressed() const
+{
+	return RealTimeTriggerWasLightlyPressed;
+}
+
+void UViewportInteractorMotionController::SetTriggerAtLeastLightlyPressed( const bool bInTriggerAtLeastLightlyPressed )
+{
+	bIsTriggerAtLeastLightlyPressed = bInTriggerAtLeastLightlyPressed;
+}
+
+void UViewportInteractorMotionController::SetTriggerBeenReleasedSinceLastPress( const bool bInTriggerBeenReleasedSinceLastPress )
+{
+	bHasTriggerBeenReleasedSinceLastPress = bInTriggerBeenReleasedSinceLastPress;
+}
+
+void UViewportInteractorMotionController::PlayHapticEffect( const float Strength )
+{
+	IInputInterface* InputInterface = FSlateApplication::Get().GetInputInterface();
+	if ( InputInterface )
+	{
+		const double CurrentTime = FPlatformTime::Seconds();
+
+		// If we're dealing with an Oculus Rift, we have to setup haptic feedback directly.  Otherwise we can use our
+		// generic force feedback system
+		if ( GetHMDDeviceType() == EHMDDeviceType::DT_OculusRift )
+		{
+			// Haptics are a little strong on Oculus Touch, so we scale them down a bit
+			const float HapticScaleForRift = 0.8f;
+
+			FHapticFeedbackValues HapticFeedbackValues;
+			HapticFeedbackValues.Amplitude = Strength * HapticScaleForRift;
+			HapticFeedbackValues.Frequency = 0.5f;
+			InputInterface->SetHapticFeedbackValues( WorldInteraction->GetMotionControllerID(), ( int32 )ControllerHandSide, HapticFeedbackValues );
+
+			LastHapticTime = CurrentTime;
+		}
+		else
+		{
+			//@todo viewportinteration
+			FForceFeedbackValues ForceFeedbackValues;
+			ForceFeedbackValues.LeftLarge = ControllerHandSide == EControllerHand::Left ? Strength : 0;
+			ForceFeedbackValues.RightLarge = ControllerHandSide == EControllerHand::Right ? Strength : 0;
+
+			// @todo vreditor: If an Xbox controller is plugged in, this causes both the motion controllers and the Xbox controller to vibrate!
+			InputInterface->SetForceFeedbackChannelValues( WorldInteraction->GetMotionControllerID(), ForceFeedbackValues );
+
+			if ( ForceFeedbackValues.LeftLarge > KINDA_SMALL_NUMBER )
+			{
+				LastHapticTime = CurrentTime;
+			}
+
+			if ( ForceFeedbackValues.RightLarge > KINDA_SMALL_NUMBER )
+			{
+				LastHapticTime = CurrentTime;
+			}
+		}
+	}
+
+	// @todo vreditor: We'll stop haptics right away because if the frame hitches, the controller will be left vibrating
+	StopOldHapticEffects();
+}
+
+void UViewportInteractorMotionController::StopOldHapticEffects()
+{
+	const float MinHapticTime = VI::MinHapticTimeForRift->GetFloat();
+
+	IInputInterface* InputInterface = FSlateApplication::Get().GetInputInterface();
+	if ( InputInterface )
+	{
+		// If we're dealing with an Oculus Rift, we have to setup haptic feedback directly.  Otherwise we can use our
+		// generic force feedback system
+		if ( GetHMDDeviceType() == EHMDDeviceType::DT_OculusRift )
+		{
+			bool bWaitingForMoreHaptics = false;
+
+			if ( LastHapticTime != 0.0f )
+			{
+				do
+				{
+					bWaitingForMoreHaptics = false;
+
+					const double CurrentTime = FPlatformTime::Seconds();
+
+					// Left hand
+					if ( CurrentTime - LastHapticTime >= MinHapticTime )
+					{
+						FHapticFeedbackValues HapticFeedbackValues;
+						HapticFeedbackValues.Amplitude = 0.0f;
+						HapticFeedbackValues.Frequency = 0.0f;
+						InputInterface->SetHapticFeedbackValues( WorldInteraction->GetMotionControllerID(), ( int32 ) ControllerHandSide, HapticFeedbackValues );
+
+						LastHapticTime = 0.0;
+					}
+					else if ( LastHapticTime != 0.0 )
+					{
+						bWaitingForMoreHaptics = true;
+					}
+
+					if ( bWaitingForMoreHaptics && VI::SleepForRiftHaptics->GetInt() != 0 )
+					{
+						FPlatformProcess::Sleep( 0 );
+					}
+				}
+				// @todo vreditor urgent: This is needed so that haptics don't play too long.  Our Oculus code doesn't currently support 
+				// multi-threading, so we need to delay the main thread to make sure we stop it before it rumbles for more than an instant!
+				while ( bWaitingForMoreHaptics && VI::SleepForRiftHaptics->GetInt() != 0 );
+			}
+		}
+		else
+		{
+			// @todo vreditor: Do we need to cancel haptics for non-Rift devices?  Doesn't seem like it
+		}
+	}
+}
+
+bool UViewportInteractorMotionController::GetTransformAndForwardVector( FTransform& OutHandTransform, FVector& OutForwardVector )
+{
+	if ( bHaveMotionController )
+	{
+		OutHandTransform = InteractorData.Transform;
+
+		const float LaserPointerRotationOffset = GetHMDDeviceType() == EHMDDeviceType::DT_OculusRift ? VI::OculusLaserPointerRotationOffset->GetFloat() : VI::ViveLaserPointerRotationOffset->GetFloat();
+		OutForwardVector = OutHandTransform.GetRotation().RotateVector( FRotator( LaserPointerRotationOffset, 0.0f, 0.0f ).RotateVector( FVector( 1.0f, 0.0f, 0.0f ) ) );
+
+		return true;
+	}
+
+	return false;
+}
+
+EControllerHand UViewportInteractorMotionController::GetControllerSide() const
+{
+	return ControllerHandSide;
+}
+
+UMotionControllerComponent* UViewportInteractorMotionController::GetMotionControllerComponent() const
+{
+	return MotionControllerComponent;
+}
+
+float UViewportInteractorMotionController::GetSelectAndMoveTriggerValue() const
+{
+	return SelectAndMoveTriggerValue;
+}
