VR Editor:

diff --git a/Engine/Source/Editor/ViewportInteraction/ViewportWorldInteraction.cpp b/Engine/Source/Editor/ViewportInteraction/ViewportWorldInteraction.cpp
new file mode 100644
--- /dev/null
+++ b/Engine/Source/Editor/ViewportInteraction/ViewportWorldInteraction.cpp
@@ -0,0 +1,2825 @@
+// Copyright 1998-2015 Epic Games, Inc. All Rights Reserved.
+
+#include "ViewportInteractionModule.h"
+#include "ViewportWorldInteraction.h"
+#include "ViewportInteractor.h"
+#include "ViewportInteractableInterface.h"
+#include "ViewportDragOperation.h"
+#include "VIGizmoHandle.h"
+#include "VIBaseTransformGizmo.h"
+#include "Gizmo/VIPivotTransformGizmo.h"
+
+#include "SnappingUtils.h"
+#include "ScopedTransaction.h"
+#include "SLevelViewport.h"
+
+// For actor placement
+#include "ObjectTools.h"
+#include "AssetSelection.h"
+
+#define LOCTEXT_NAMESPACE "ViewportWorldInteraction"
+
+namespace VI
+{
+	static FAutoConsoleVariable ScaleSensitivity( TEXT( "VI.ScaleSensitivity" ), 0.01f, TEXT( "Sensitivity for scaling" ) );
+	static FAutoConsoleVariable GizmoRotationSensitivity( TEXT( "VI.GizmoRotationSensitivity" ), 0.25f, TEXT( "How much to rotate as the user drags on a rotation gizmo handle" ) );
+	static FAutoConsoleVariable ScaleWorldFromFloor( TEXT( "VI.ScaleWorldFromFloor" ), 0, TEXT( "Whether the world should scale relative to your tracking space floor instead of the center of your hand locations" ) );
+	static FAutoConsoleVariable AllowVerticalWorldMovement( TEXT( "VI.AllowVerticalWorldMovement" ), 1, TEXT( "Whether you can move your tracking space away from the origin or not" ) );
+	static FAutoConsoleVariable AllowWorldRotationPitchAndRoll( TEXT( "VI.AllowWorldRotationPitchAndRoll" ), 0, TEXT( "When enabled, you'll not only be able to yaw, but also pitch and roll the world when rotating by gripping with two hands" ) );
+	static FAutoConsoleVariable InertiaVelocityBoost( TEXT( "VI.InertiaVelocityBoost" ), 3000.0f, TEXT( "How much to scale object velocity when releasing dragged simulating objects in Simulate mode" ) );
+	static FAutoConsoleVariable SweepPhysicsWhileSimulating( TEXT( "VI.SweepPhysicsWhileSimulating" ), 1, TEXT( "If enabled, simulated objects won't be able to penetrate other objects while being dragged in Simulate mode" ) );
+	static FAutoConsoleVariable PlacementInterpolationDuration( TEXT( "VI.PlacementInterpolationDuration" ), 0.6f, TEXT( "How long we should interpolate newly-placed objects to their target location." ) );
+	static FAutoConsoleVariable SmoothSnap( TEXT( "VI.SmoothSnap" ), 1, TEXT( "When enabled with grid snap, transformed objects will smoothly blend to their new location (instead of teleporting instantly)" ) );
+	static FAutoConsoleVariable SmoothSnapSpeed( TEXT( "VI.SmoothSnapSpeed" ), 30.0f, TEXT( "How quickly objects should interpolate to their new position when grid snapping is enabled" ) );
+	static FAutoConsoleVariable ElasticSnap( TEXT( "VI.ElasticSnap" ), 1, TEXT( "When enabled with grid snap, you can 'pull' objects slightly away from their snapped position" ) );
+	static FAutoConsoleVariable ElasticSnapStrength( TEXT( "VI.ElasticSnapStrength" ), 0.3f, TEXT( "How much objects should 'reach' toward their unsnapped position when elastic snapping is enabled with grid snap" ) );
+	static FAutoConsoleVariable ScaleMax( TEXT( "VI.ScaleMax" ), 6000.0f, TEXT( "Maximum world scale in centimeters" ) );
+	static FAutoConsoleVariable ScaleMin( TEXT( "VI.ScaleMin" ), 10.0f, TEXT( "Minimum world scale in centimeters" ) );
+	static FAutoConsoleVariable DragAtLaserImpactInterpolationDuration( TEXT( "VI.DragAtLaserImpactInterpolationDuration" ), 0.1f, TEXT( "How long we should interpolate objects between positions when dragging under the laser's impact point" ) );
+	static FAutoConsoleVariable DragAtLaserImpactInterpolationThreshold( TEXT( "VI.DragAtLaserImpactInterpolationThreshold" ), 5.0f, TEXT( "Minimum distance jumped between frames before we'll force interpolation mode to activated" ) );
+
+	static FAutoConsoleVariable ForceGizmoPivotToCenterOfSelectedActorsBounds( TEXT( "VI.ForceGizmoPivotToCenterOfSelectedActorsBounds" ), 0, TEXT( "When enabled, the gizmo's pivot will always be centered on the selected actors.  Otherwise, we use the pivot of the last selected actor." ) );
+	static FAutoConsoleVariable GizmoHandleHoverScale( TEXT( "VI.GizmoHandleHoverScale" ), 1.5f, TEXT( "How much to scale up transform gizmo handles when hovered over" ) );
+	static FAutoConsoleVariable GizmoHandleHoverAnimationDuration( TEXT( "VI.GizmoHandleHoverAnimationDuration" ), 0.1f, TEXT( "How quickly to animate gizmo handle hover state" ) );
+	static FAutoConsoleVariable ShowTransformGizmo( TEXT( "VI.ShowTransformGizmo" ), 1, TEXT( "Whether the transform gizmo should be shown for selected objects" ) );
+
+	//Grid
+	static FAutoConsoleVariable SnapGridSize( TEXT( "VI.SnapGridSize" ), 3.0f, TEXT( "How big the snap grid should be, in multiples of the gizmo bounding box size" ) );
+	static FAutoConsoleVariable SnapGridLineWidth( TEXT( "VI.SnapGridLineWidth" ), 3.0f, TEXT( "Width of the grid lines on the snap grid" ) );
+	static FAutoConsoleVariable MinVelocityForInertia( TEXT( "VREd.MinVelocityForInertia" ), 1.0f, TEXT( "Minimum velocity (in cm/frame in unscaled room space) before inertia will kick in when releasing objects (or the world)" ) );
+	static FAutoConsoleVariable GridHapticFeedbackStrength( TEXT( "VI.GridHapticFeedbackStrength" ), 0.4f, TEXT( "Default strength for haptic feedback when moving across grid points" ) );
+}
+
+// @todo vreditor: Hacky inline implementation of a double vector.  Move elsewhere and polish or get rid.
+struct DVector
+{
+	double X;
+	double Y;
+	double Z;
+
+	DVector()
+	{
+	}
+
+	DVector( const DVector& V )
+		: X( V.X ), Y( V.Y ), Z( V.Z )
+	{
+	}
+
+	DVector( const FVector& V )
+		: X( V.X ), Y( V.Y ), Z( V.Z )
+	{
+	}
+
+	DVector( double InX, double InY, double InZ )
+		: X( InX ), Y( InY ), Z( InZ )
+	{
+	}
+
+	DVector operator+( const DVector& V ) const
+	{
+		return DVector( X + V.X, Y + V.Y, Z + V.Z );
+	}
+
+	DVector operator-( const DVector& V ) const
+	{
+		return DVector( X - V.X, Y - V.Y, Z - V.Z );
+	}
+
+	DVector operator*( double Scale ) const
+	{
+		return DVector( X * Scale, Y * Scale, Z * Scale );
+	}
+
+	double operator|( const DVector& V ) const
+	{
+		return X*V.X + Y*V.Y + Z*V.Z;
+	}
+
+	FVector ToFVector() const
+	{
+		return FVector( ( float ) X, ( float ) Y, ( float ) Z );
+	}
+};
+
+// @todo vreditor: Move elsewhere or get rid
+static void SegmentDistToSegmentDouble( DVector A1, DVector B1, DVector A2, DVector B2, DVector& OutP1, DVector& OutP2 )
+{
+	const double Epsilon = 1.e-10;
+
+	// Segments
+	const	DVector	S1 = B1 - A1;
+	const	DVector	S2 = B2 - A2;
+	const	DVector	S3 = A1 - A2;
+
+	const	double	Dot11 = S1 | S1;	// always >= 0
+	const	double	Dot22 = S2 | S2;	// always >= 0
+	const	double	Dot12 = S1 | S2;
+	const	double	Dot13 = S1 | S3;
+	const	double	Dot23 = S2 | S3;
+
+	// Numerator
+	double	N1, N2;
+
+	// Denominator
+	const	double	D = Dot11*Dot22 - Dot12*Dot12;	// always >= 0
+	double	D1 = D;		// T1 = N1 / D1, default D1 = D >= 0
+	double	D2 = D;		// T2 = N2 / D2, default D2 = D >= 0
+
+						// compute the line parameters of the two closest points
+	if ( D < Epsilon )
+	{
+		// the lines are almost parallel
+		N1 = 0.0;	// force using point A on segment S1
+		D1 = 1.0;	// to prevent possible division by 0 later
+		N2 = Dot23;
+		D2 = Dot22;
+	}
+	else
+	{
+		// get the closest points on the infinite lines
+		N1 = ( Dot12*Dot23 - Dot22*Dot13 );
+		N2 = ( Dot11*Dot23 - Dot12*Dot13 );
+
+		if ( N1 < 0.0 )
+		{
+			// t1 < 0.0 => the s==0 edge is visible
+			N1 = 0.0;
+			N2 = Dot23;
+			D2 = Dot22;
+		}
+		else if ( N1 > D1 )
+		{
+			// t1 > 1 => the t1==1 edge is visible
+			N1 = D1;
+			N2 = Dot23 + Dot12;
+			D2 = Dot22;
+		}
+	}
+
+	if ( N2 < 0.0 )
+	{
+		// t2 < 0 => the t2==0 edge is visible
+		N2 = 0.0;
+
+		// recompute t1 for this edge
+		if ( -Dot13 < 0.0 )
+		{
+			N1 = 0.0;
+		}
+		else if ( -Dot13 > Dot11 )
+		{
+			N1 = D1;
+		}
+		else
+		{
+			N1 = -Dot13;
+			D1 = Dot11;
+		}
+	}
+	else if ( N2 > D2 )
+	{
+		// t2 > 1 => the t2=1 edge is visible
+		N2 = D2;
+
+		// recompute t1 for this edge
+		if ( ( -Dot13 + Dot12 ) < 0.0 )
+		{
+			N1 = 0.0;
+		}
+		else if ( ( -Dot13 + Dot12 ) > Dot11 )
+		{
+			N1 = D1;
+		}
+		else
+		{
+			N1 = ( -Dot13 + Dot12 );
+			D1 = Dot11;
+		}
+	}
+
+	// finally do the division to get the points' location
+	const double T1 = ( FMath::Abs( N1 ) < Epsilon ? 0.0 : N1 / D1 );
+	const double T2 = ( FMath::Abs( N2 ) < Epsilon ? 0.0 : N2 / D2 );
+
+	// return the closest points
+	OutP1 = A1 + S1 * T1;
+	OutP2 = A2 + S2 * T2;
+}
+
+UViewportWorldInteraction::UViewportWorldInteraction( const FObjectInitializer& Initializer ) :
+	Super( Initializer ),
+	AppTimeEntered( FTimespan::Zero() ),
+	LastFrameNumberInputWasPolled( 0 ),
+	MotionControllerID( 0 ),	// @todo ViewportInteraction: We only support a single controller, and we assume the first controller are the motion controls
+	EditorViewportClient( nullptr ),
+	LastWorldToMetersScale( 100.0f ),
+	bIsInterpolatingTransformablesFromSnapshotTransform( false ),
+	TransformablesInterpolationStartTime( FTimespan::Zero() ),
+	TransformablesInterpolationDuration( 1.0f ),
+	TransformGizmoActor( nullptr ),
+	TransformGizmoClass( APivotTransformGizmo::StaticClass() ),
+	GizmoLocalBounds( FBox( 0 ) ),
+	bDraggedSinceLastSelection( false ),
+	LastDragGizmoStartTransform( FTransform::Identity ),
+	StartDragAngleOnRotation( ),
+	StartDragHandleDirection( ),
+	CurrentGizmoType( EGizmoHandleTypes::All ),
+	bIsUsingVR( false ),
+	SnapGridActor( nullptr ),
+	SnapGridMeshComponent( nullptr ),
+	SnapGridMID( nullptr ),
+	DraggedInteractable( nullptr )
+{
+	InputProcessor = MakeShareable( new FViewportInteractionInputProcessor( *this ) );
+	FSlateApplication::Get().SetInputPreProcessor( true, InputProcessor );
+
+	// Find out about selection changes
+	USelection::SelectionChangedEvent.AddUObject( this, &UViewportWorldInteraction::OnActorSelectionChanged ); //@todo viewportinteraction
+
+	// Setting up colors
+	Colors.SetNumZeroed( ( int32 ) EColors::TotalCount );
+	{
+		Colors[ ( int32 ) EColors::DefaultColor ] = FLinearColor( 0.7f, 0.7f, 0.7f, 1.0f );;
+		Colors[ ( int32 ) EColors::Forward ] = FLinearColor( 0.4f, 0.05f, 0.05f, 1.0f );
+		Colors[ ( int32 ) EColors::Right ] = FLinearColor( 0.05f, 0.4f, 0.05f, 1.0f );
+		Colors[ ( int32 ) EColors::Up ] = FLinearColor( 0.05f, 0.05f, 0.4f, 1.0f );
+		Colors[ ( int32 ) EColors::Hover ] = FLinearColor::Yellow;
+		Colors[ ( int32 ) EColors::Dragging ] = FLinearColor::White;
+	}
+}
+
+void UViewportWorldInteraction::Init( FEditorViewportClient* InEditorViewportClient )
+{
+	EditorViewportClient = InEditorViewportClient;
+	AppTimeEntered = FTimespan::FromSeconds( FApp::GetCurrentTime() );
+
+	//Spawn the transform gizmo at init so we do not hitch when selecting our first object
+	SpawnTransformGizmoIfNeeded();
+}
+
+void UViewportWorldInteraction::Shutdown()
+{
+	AppTimeEntered = FTimespan::Zero();
+
+	OnHoverUpdateEvent.Clear();
+	OnInputActionEvent.Clear();
+
+	if ( TransformGizmoActor != nullptr )
+	{
+		DestroyTransientActor( TransformGizmoActor );
+		TransformGizmoActor = nullptr;
+	}
+
+	if ( SnapGridActor != nullptr )
+	{
+		DestroyTransientActor( SnapGridActor );
+		SnapGridActor = nullptr;
+		SnapGridMeshComponent = nullptr;
+	}
+
+	if ( SnapGridMID != nullptr )
+	{
+		SnapGridMID->MarkPendingKill();
+		SnapGridMID = nullptr;
+	}
+
+	for ( UViewportInteractor* Interactor : Interactors )
+	{
+		Interactor->Shutdown();
+	}
+	Interactors.Empty();
+
+	USelection::SelectionChangedEvent.RemoveAll( this );
+}
+
+void UViewportWorldInteraction::Tick( FEditorViewportClient* ViewportClient, const float DeltaTime )
+{
+	// Only if this is our viewport. Remember, editor modes get a chance to tick and receive input for each active viewport.
+	if ( ViewportClient != EditorViewportClient )
+	{
+		return;
+	}
+
+	// Cancel any haptic effects that have been playing too long
+	//StopOldHapticEffects(); //@todo ViewportInteraction
+
+	// Get the latest interactor data, and fill in our its data with fresh transforms
+	PollInputIfNeeded();
+
+	// Update hover. Note that hover can also be updated when ticking our sub-systems below.
+	HoverTick( ViewportClient, DeltaTime );
+
+	InteractionTick( ViewportClient, DeltaTime );
+
+	// Update all the interactors
+	for ( UViewportInteractor* Interactor : Interactors )
+	{
+		Interactor->Tick( DeltaTime );
+	}
+}
+
+void UViewportWorldInteraction::AddInteractor( UViewportInteractor* Interactor )
+{
+	Interactor->SetWorldInteraction( this );
+	Interactors.AddUnique( Interactor );
+}
+
+void UViewportWorldInteraction::RemoveInteractor( UViewportInteractor* Interactor )
+{
+	Interactor->GetOtherInteractor()->RemoveOtherInteractor();
+	Interactor->RemoveOtherInteractor();
+	Interactors.Remove( Interactor );
+}
+
+void UViewportWorldInteraction::PairInteractors( UViewportInteractor* FirstInteractor, UViewportInteractor* SecondInteractor )
+{
+	FirstInteractor->SetOtherInteractor( SecondInteractor );
+	SecondInteractor->SetOtherInteractor( FirstInteractor );
+}
+
+bool UViewportWorldInteraction::HandleInputKey( const FKey Key, const EInputEvent Event )
+{
+	bool bWasHandled = false;
+
+	for ( UViewportInteractor* Interactor : Interactors )
+	{
+		// Stop iterating if the input was handled by an Interactor
+		if ( bWasHandled )
+		{
+			break;
+		}
+
+		bWasHandled = Interactor->HandleInputKey( Key, Event );
+	}
+
+	return bWasHandled;
+}
+
+bool UViewportWorldInteraction::HandleInputAxis( const int32 ControllerId, const FKey Key, const float Delta, const float DeltaTime )
+{
+	bool bWasHandled = false;
+
+	for ( UViewportInteractor* Interactor : Interactors )
+	{
+		// Stop iterating if the input was handled by an interactor
+		if ( bWasHandled )
+		{
+			break;
+		}
+
+		bWasHandled = Interactor->HandleInputAxis( Key, Delta, DeltaTime );
+	}
+
+	return bWasHandled;
+}
+
+
+FTransform UViewportWorldInteraction::GetRoomTransform() const
+{
+	const FTransform RoomTransform(
+		EditorViewportClient->GetViewRotation().Quaternion(),
+		EditorViewportClient->GetViewLocation(),
+		FVector( 1.0f ) );
+	return RoomTransform;
+}
+
+FTransform UViewportWorldInteraction::GetRoomSpaceHeadTransform() const
+{
+	FTransform HeadTransform = FTransform::Identity;
+	if ( bIsUsingVR && GEngine->HMDDevice.IsValid() )
+	{
+		FQuat RoomSpaceHeadOrientation;
+		FVector RoomSpaceHeadLocation;
+		GEngine->HMDDevice->GetCurrentOrientationAndPosition( /* Out */ RoomSpaceHeadOrientation, /* Out */ RoomSpaceHeadLocation );
+
+		HeadTransform = FTransform(
+			RoomSpaceHeadOrientation,
+			RoomSpaceHeadLocation,
+			FVector( 1.0f ) );
+	}
+
+	return HeadTransform;
+}
+
+
+FTransform UViewportWorldInteraction::GetHeadTransform() const
+{
+	return GetRoomSpaceHeadTransform() * GetRoomTransform();
+}
+
+void UViewportWorldInteraction::SetRoomTransform( const FTransform& NewRoomTransform )
+{
+	EditorViewportClient->SetViewLocation( NewRoomTransform.GetLocation() );
+	EditorViewportClient->SetViewRotation( NewRoomTransform.GetRotation().Rotator() );
+
+	// Forcibly dirty the viewport camera location
+	const bool bDollyCamera = false;
+	EditorViewportClient->MoveViewportCamera( FVector::ZeroVector, FRotator::ZeroRotator, bDollyCamera );
+}
+
+float UViewportWorldInteraction::GetWorldScaleFactor() const
+{
+	return GetViewportWorld()->GetWorldSettings()->WorldToMeters / 100.0f;
+}
+
+UWorld* UViewportWorldInteraction::GetViewportWorld() const
+{
+	return EditorViewportClient->GetWorld();
+}
+
+FEditorViewportClient* UViewportWorldInteraction::GetViewportClient() const
+{
+	return EditorViewportClient;
+}
+
+void UViewportWorldInteraction::SetIsUsingVR( const bool bInIsUsingVR )
+{
+	bIsUsingVR = bInIsUsingVR;
+}
+
+void UViewportWorldInteraction::Undo()
+{
+	GUnrealEd->Exec( GetViewportWorld(), TEXT( "TRANSACTION UNDO" ) );
+}
+
+void UViewportWorldInteraction::Redo()
+{
+	GUnrealEd->Exec( GetViewportWorld(), TEXT( "TRANSACTION REDO" ) );
+}
+
+void UViewportWorldInteraction::DeleteSelectedObjects()
+{
+	const FScopedTransaction Transaction( LOCTEXT( "DeleteSelection", "Delete selection" ) );
+	
+	GUnrealEd->edactDeleteSelected( GetViewportWorld() );
+}
+
+void UViewportWorldInteraction::Copy()
+{
+	// @todo vreditor: Needs CanExecute()  (see LevelEditorActions.cpp)
+	GUnrealEd->Exec( GetViewportWorld(), TEXT( "EDIT COPY" ) );
+}
+
+void UViewportWorldInteraction::Paste()
+{
+	// @todo vreditor: Needs CanExecute()  (see LevelEditorActions.cpp)
+	// @todo vreditor: Needs "paste here" style pasting (TO=HERE), but with ray
+	GUnrealEd->Exec( GetViewportWorld(), TEXT( "EDIT PASTE" ) );
+}
+
+void UViewportWorldInteraction::Duplicate()
+{
+	ABrush::SetSuppressBSPRegeneration( true );
+	GEditor->edactDuplicateSelected( GetViewportWorld()->GetCurrentLevel(), false );
+	ABrush::SetSuppressBSPRegeneration( false );
+}
+
+void UViewportWorldInteraction::Deselect()
+{
+	GEditor->SelectNone( true, true );
+}
+
+void UViewportWorldInteraction::HoverTick( FEditorViewportClient* ViewportClient, const float DeltaTime )
+{
+	for ( UViewportInteractor* Interactor : Interactors )
+	{
+		bool bWasHandled = false;
+		FViewportInteractorData& InteractorData = Interactor->GetInteractorData();
+		
+		Interactor->ResetHoverState( DeltaTime );
+		
+		//Update the state of the interactor here
+		FHitResult HitHoverResult;
+		Interactor->UpdateHoverResult( HitHoverResult );
+
+		FVector HoverImpactPoint = HitHoverResult.ImpactPoint;
+		OnHoverUpdateEvent.Broadcast( *ViewportClient, Interactor, /* In/Out */ HoverImpactPoint, /* In/Out */ bWasHandled);
+		
+		if ( bWasHandled )
+		{
+			InteractorData.bIsHovering = true;
+			InteractorData.HoverLocation = HoverImpactPoint;
+		}
+		else
+		{
+			UActorComponent* NewHoveredActorComponent = nullptr;
+
+			InteractorData.HoveringOverTransformGizmoComponent = nullptr;
+			InteractorData.HoverLocation = FVector::ZeroVector;
+
+			FHitResult HitResult = Interactor->GetHitResultFromLaserPointer();
+			if ( HitResult.Actor.IsValid() )
+			{
+				USceneComponent* HoveredActorComponent = HitResult.GetComponent();
+				AActor* Actor = HitResult.Actor.Get();
+
+				if ( HoveredActorComponent && IsInteractableComponent( HoveredActorComponent ) )
+				{
+					HoveredObjects.Add( FViewportHoverTarget( Actor ) );
+
+					InteractorData.HoverLocation = HitResult.ImpactPoint;
+
+					if ( Actor == TransformGizmoActor )
+					{
+						NewHoveredActorComponent = HoveredActorComponent;
+						InteractorData.HoveringOverTransformGizmoComponent = HoveredActorComponent;
+						InteractorData.HoverHapticCheckLastHoveredGizmoComponent = HitResult.GetComponent();
+					}
+					else
+					{
+						IViewportInteractableInterface* Interactable = Cast<IViewportInteractableInterface>( Actor );
+						if ( Interactable )
+						{
+							NewHoveredActorComponent = HoveredActorComponent;
+
+							// Check if the current hovered component of the interactor is different from the hitresult component
+							if ( NewHoveredActorComponent != InteractorData.HoveredActorComponent && !IsOtherInteractorHoveringOverComponent( Interactor, NewHoveredActorComponent ) )
+							{
+								Interactable->OnHoverEnter( Interactor, HitResult );
+							}
+
+							Interactable->OnHover( Interactor );
+						}
+					}
+					InteractorData.bIsHovering = true;
+				}
+				else if ( !HoveredActorComponent )
+				{
+					InteractorData.bIsHovering = false;
+				}
+			}
+
+			// Leave hovered interactable
+			//@todo ViewportInteraction: This does not take into account when the other interactors are already hovering over this interactable
+			if ( InteractorData.HoveredActorComponent != nullptr && ( InteractorData.HoveredActorComponent != NewHoveredActorComponent || NewHoveredActorComponent == nullptr ) 
+				&& !IsOtherInteractorHoveringOverComponent( Interactor, NewHoveredActorComponent ) )
+			{
+				AActor* HoveredActor = InteractorData.HoveredActorComponent->GetOwner();
+				if ( HoveredActor )
+				{
+					IViewportInteractableInterface* Interactable = Cast<IViewportInteractableInterface>( HoveredActor );
+					if ( Interactable )
+					{
+						Interactable->OnHoverLeave( Interactor, NewHoveredActorComponent );
+					}
+				}
+			}
+
+			//Update the hovered actor component with the new component
+			InteractorData.HoveredActorComponent = NewHoveredActorComponent;
+		}
+
+		
+	}
+}
+
+void UViewportWorldInteraction::InteractionTick( FEditorViewportClient* ViewportClient, const float DeltaTime )
+{
+	const FTimespan CurrentTime = FTimespan::FromSeconds( FPlatformTime::Seconds() );
+	const float WorldToMetersScale = GetViewportWorld()->GetWorldSettings()->WorldToMeters;
+
+	// Update viewport with any objects that are currently hovered over
+	{
+		const bool bUseEditorSelectionHoverFeedback = GEditor != NULL && GetDefault<ULevelEditorViewportSettings>()->bEnableViewportHoverFeedback;
+		if( bUseEditorSelectionHoverFeedback )
+		{
+			if( ViewportClient->IsLevelEditorClient() )
+			{
+				FLevelEditorViewportClient* LevelEditorViewportClient = static_cast<FLevelEditorViewportClient*>( ViewportClient );
+				LevelEditorViewportClient->UpdateHoveredObjects( HoveredObjects );
+			}
+		}
+
+		// This will be filled in again during the next input update
+		HoveredObjects.Reset();
+	}
+
+	const float WorldScaleFactor = GetWorldScaleFactor();
+
+	// Move selected actors
+	TOptional< UViewportInteractor* > DraggingActorsWithInteractor;
+	TOptional< UViewportInteractor* > AssistingDragWithInteractor;
+
+	for ( UViewportInteractor* Interactor : Interactors )
+	{
+		bool bCanSlideRayLength = false;
+		FViewportInteractorData& InteractorData = Interactor->GetInteractorData();
+
+		if ( InteractorData.DraggingMode == EViewportInteractionDraggingMode::ActorsWithGizmo ||
+			InteractorData.DraggingMode == EViewportInteractionDraggingMode::ActorsFreely ||
+			InteractorData.DraggingMode == EViewportInteractionDraggingMode::ActorsAtLaserImpact )
+		{
+			check( !DraggingActorsWithInteractor.IsSet() );	// Only support dragging one thing at a time right now!
+			DraggingActorsWithInteractor = Interactor;
+
+			if ( InteractorData.DraggingMode != EViewportInteractionDraggingMode::ActorsAtLaserImpact )
+			{
+				bCanSlideRayLength = true;
+			}
+		}
+
+		if ( InteractorData.DraggingMode == EViewportInteractionDraggingMode::AssistingDrag )
+		{
+			check( !AssistingDragWithInteractor.IsSet() );	// Only support assisting one thing at a time right now!
+			AssistingDragWithInteractor = Interactor;
+
+			bCanSlideRayLength = true;
+		}
+
+		if ( bCanSlideRayLength )
+		{
+			Interactor->CalculateDragRay();
+		}
+	}
+	
+	for ( UViewportInteractor* Interactor : Interactors )
+	{
+		FViewportInteractorData& InteractorData = Interactor->GetInteractorData();
+			
+		if( InteractorData.DraggingMode == EViewportInteractionDraggingMode::ActorsWithGizmo ||
+			InteractorData.DraggingMode == EViewportInteractionDraggingMode::ActorsFreely ||
+			InteractorData.DraggingMode == EViewportInteractionDraggingMode::ActorsAtLaserImpact ||
+			InteractorData.DraggingMode == EViewportInteractionDraggingMode::World )
+		{
+			// Are we dragging with two interactors ?
+			UViewportInteractor* OtherInteractor = nullptr;
+			if ( AssistingDragWithInteractor.IsSet() )
+			{
+				OtherInteractor = AssistingDragWithInteractor.GetValue();
+			}
+
+			FViewportInteractorData& OtherInteractorData = OtherInteractor->GetInteractorData();
+
+			UViewportInteractor* OtherInteractorThatWasAssistingDrag = GetOtherInteractorIntertiaContribute( Interactor );
+
+			FVector DraggedTo = InteractorData.Transform.GetLocation();
+			FVector DragDelta = DraggedTo - InteractorData.LastTransform.GetLocation();
+			FVector DragDeltaFromStart = DraggedTo - InteractorData.LaserPointerImpactAtDragStart;
+
+			FVector OtherHandDraggedTo = FVector::ZeroVector;
+			FVector OtherHandDragDelta = FVector::ZeroVector;
+			FVector OtherHandDragDeltaFromStart = FVector::ZeroVector;
+			if( OtherInteractor != nullptr )
+			{
+				OtherHandDraggedTo = OtherInteractor->GetInteractorData().Transform.GetLocation();
+				OtherHandDragDelta = OtherHandDraggedTo - OtherInteractor->GetInteractorData().LastTransform.GetLocation();
+				OtherHandDragDeltaFromStart = DraggedTo - OtherInteractor->GetInteractorData().LaserPointerImpactAtDragStart;
+			}
+			else if( OtherInteractorThatWasAssistingDrag != nullptr )
+			{
+				OtherHandDragDelta = OtherInteractorThatWasAssistingDrag->GetInteractorData().DragTranslationVelocity;
+				OtherHandDraggedTo = OtherInteractorThatWasAssistingDrag->GetInteractorData().LastDragToLocation + OtherHandDragDelta;
+				OtherHandDragDeltaFromStart = OtherHandDraggedTo - OtherInteractorThatWasAssistingDrag->GetInteractorData().LaserPointerImpactAtDragStart;
+			}
+
+
+			FVector LaserPointerStart = InteractorData.Transform.GetLocation();
+			FVector LaserPointerDirection = InteractorData.Transform.GetUnitAxis( EAxis::X );
+			FVector LaserPointerEnd = InteractorData.Transform.GetLocation();
+
+			if( InteractorData.DraggingMode == EViewportInteractionDraggingMode::ActorsWithGizmo ||
+				InteractorData.DraggingMode == EViewportInteractionDraggingMode::ActorsFreely ||
+				InteractorData.DraggingMode == EViewportInteractionDraggingMode::ActorsAtLaserImpact )
+			{
+				// Move objects using the laser pointer (in world space)
+				if( Interactor->GetLaserPointer( /* Out */ LaserPointerStart, /* Out */ LaserPointerEnd ) )
+				{
+					LaserPointerDirection = ( LaserPointerEnd - LaserPointerStart ).GetSafeNormal();
+
+					if( InteractorData.DraggingMode == EViewportInteractionDraggingMode::ActorsAtLaserImpact )
+					{
+						// Check to see if the laser pointer is over something we can drop on
+						FVector HitLocation = FVector::ZeroVector;
+						bool bHitSomething = FindPlacementPointUnderLaser( Interactor, /* Out */ HitLocation );
+
+						if( bHitSomething )
+						{
+							// If the object moved reasonably far between frames, it might be because the angle we were aligning
+							// the object with during placement changed radically.  To avoid it popping, we smoothly interpolate
+							// it's position over a very short timespan
+							if( !bIsInterpolatingTransformablesFromSnapshotTransform )	// Let the last animation finish first
+							{
+								const float ScaledDragDistance = DragDelta.Size() * WorldScaleFactor;
+								if( ScaledDragDistance >= VI::DragAtLaserImpactInterpolationThreshold->GetFloat() )
+								{
+									bIsInterpolatingTransformablesFromSnapshotTransform = true;
+									TransformablesInterpolationStartTime = CurrentTime;
+									TransformablesInterpolationDuration = VI::DragAtLaserImpactInterpolationDuration->GetFloat();
+
+									// Snapshot everything
+									for( TUniquePtr<FViewportTransformable>& TransformablePtr : Transformables )
+									{
+										FViewportTransformable& Transformable = *TransformablePtr;
+										Transformable.InterpolationSnapshotTransform = Transformable.LastTransform;
+									}
+								}
+							}
+
+							DraggedTo = HitLocation;
+						}
+						else
+						{
+							DraggedTo = LaserPointerStart + LaserPointerDirection * InteractorData.DragRayLength;
+						}
+					}
+					else
+					{
+						DraggedTo = LaserPointerStart + LaserPointerDirection * InteractorData.DragRayLength;
+					}
+
+					const FVector WorldSpaceDragDelta = DraggedTo - InteractorData.LastDragToLocation;
+					DragDelta = WorldSpaceDragDelta;
+					InteractorData.DragTranslationVelocity = WorldSpaceDragDelta;
+
+					const FVector WorldSpaceDeltaFromStart = DraggedTo - InteractorData.LaserPointerImpactAtDragStart;
+					DragDeltaFromStart = WorldSpaceDeltaFromStart;
+
+					InteractorData.LastLaserPointerStart = LaserPointerStart;
+					InteractorData.LastDragToLocation = DraggedTo;
+
+					// Update hover location (we only do this when dragging using the laser pointer)
+					InteractorData.bIsHovering = true;
+					InteractorData.HoverLocation = FMath::ClosestPointOnLine( LaserPointerStart, LaserPointerEnd, DraggedTo );
+				}
+				else
+				{
+					// We lost our laser pointer, so cancel the drag
+					StopDragging( Interactor );
+				}
+
+				if( OtherInteractor != nullptr )
+				{
+					FVector OtherHandLaserPointerStart, OtherHandLaserPointerEnd;
+					if( OtherInteractor->GetLaserPointer( /* Out */ OtherHandLaserPointerStart, /* Out */ OtherHandLaserPointerEnd ) )
+					{
+						const FVector OtherHandLaserPointerDirection = ( OtherHandLaserPointerEnd - OtherHandLaserPointerStart ).GetSafeNormal();
+						OtherHandDraggedTo = OtherHandLaserPointerStart + OtherHandLaserPointerDirection * OtherInteractorData.DragRayLength;
+
+						const FVector OtherHandWorldSpaceDragDelta = OtherHandDraggedTo - OtherInteractorData.LastDragToLocation;
+						OtherHandDragDelta = OtherHandWorldSpaceDragDelta;
+						OtherInteractorData.DragTranslationVelocity = OtherHandWorldSpaceDragDelta;
+
+						const FVector OtherHandWorldSpaceDeltaFromStart = OtherHandDraggedTo - OtherInteractorData.LaserPointerImpactAtDragStart;
+						OtherHandDragDeltaFromStart = OtherHandWorldSpaceDeltaFromStart;
+
+						OtherInteractorData.LastLaserPointerStart = OtherHandLaserPointerStart;
+						OtherInteractorData.LastDragToLocation = OtherHandDraggedTo;
+
+						// Only hover if we're using the laser pointer
+						OtherInteractorData.bIsHovering = true;
+						OtherInteractorData.HoverLocation = OtherHandDraggedTo;
+					}
+					else
+					{
+						// We lost our laser pointer, so cancel the drag assist
+						StopDragging( OtherInteractor );
+					}
+				}
+			}
+			else if( ensure( InteractorData.DraggingMode == EViewportInteractionDraggingMode::World ) )
+			{
+				// While we're changing WorldToMetersScale every frame, our room-space hand locations will be scaled as well!  We need to
+				// inverse compensate for this scaling so that we can figure out how much the hands actually moved as if no scale happened.
+				// This only really matters when we're performing world scaling interactively.
+				const FVector RoomSpaceUnscaledHandLocation = ( InteractorData.RoomSpaceTransform.GetLocation() / WorldToMetersScale ) * LastWorldToMetersScale;
+				const FVector RoomSpaceUnscaledHandDelta = ( RoomSpaceUnscaledHandLocation - InteractorData.LastRoomSpaceTransform.GetLocation() );
+
+				// Move the world (in room space)
+				DraggedTo = InteractorData.LastRoomSpaceTransform.GetLocation() + RoomSpaceUnscaledHandDelta;
+
+				InteractorData.DragTranslationVelocity = RoomSpaceUnscaledHandDelta;
+				DragDelta = RoomSpaceUnscaledHandDelta;
+
+				InteractorData.LastDragToLocation = DraggedTo;
+
+				// Two handed?
+				if( OtherInteractor != nullptr )
+				{
+					const FVector OtherHandRoomSpaceUnscaledLocation = ( OtherInteractorData.RoomSpaceTransform.GetLocation() / WorldToMetersScale ) * LastWorldToMetersScale;
+					const FVector OtherHandRoomSpaceUnscaledHandDelta = ( OtherHandRoomSpaceUnscaledLocation - OtherInteractorData.LastRoomSpaceTransform.GetLocation() );
+
+					OtherHandDraggedTo = OtherInteractorData.LastRoomSpaceTransform.GetLocation() + OtherHandRoomSpaceUnscaledHandDelta;
+
+					OtherInteractorData.DragTranslationVelocity = OtherHandRoomSpaceUnscaledHandDelta;
+					OtherHandDragDelta = OtherHandRoomSpaceUnscaledHandDelta;
+
+					OtherInteractorData.LastDragToLocation = OtherHandDraggedTo;
+				}
+			}
+
+			{
+				// Intertia must only work in VR
+				// @todo vreditor: We could do a world space distance test (times world scale factor) when forcing VR mode to get similar (but not quite the same) behavior
+				if( bIsUsingVR )
+				{
+					// Don't bother with inertia if we're not moving very fast.  This filters out tiny accidental movements.
+					const FVector RoomSpaceHandDelta = ( InteractorData.RoomSpaceTransform.GetLocation() - InteractorData.LastRoomSpaceTransform.GetLocation() );
+					if( RoomSpaceHandDelta.Size() < VI::MinVelocityForInertia->GetFloat() * WorldScaleFactor )
+					{
+						InteractorData.DragTranslationVelocity = FVector::ZeroVector;
+					}
+					if( AssistingDragWithInteractor.IsSet() )
+					{
+						FViewportInteractorData& AssistingOtherInteractorData = AssistingDragWithInteractor.GetValue()->GetInteractorData();
+						const FVector OtherHandRoomSpaceHandDelta = ( AssistingOtherInteractorData.RoomSpaceTransform.GetLocation() - AssistingOtherInteractorData.LastRoomSpaceTransform.GetLocation() );
+						if( OtherHandRoomSpaceHandDelta.Size() < VI::MinVelocityForInertia->GetFloat() * WorldScaleFactor )
+						{
+							AssistingOtherInteractorData.DragTranslationVelocity = FVector::ZeroVector;
+						}
+					}
+				}
+			}
+
+			const FVector OldViewLocation = ViewportClient->GetViewLocation();
+			// Dragging transform gizmo handle
+			const bool bWithTwoHands = ( OtherInteractor != nullptr || OtherInteractorThatWasAssistingDrag != nullptr );
+			const bool bIsLaserPointerValid = true;
+			FVector UnsnappedDraggedTo = FVector::ZeroVector;
+			UpdateDragging(
+				DeltaTime,
+				InteractorData.bIsFirstDragUpdate, 
+				InteractorData.DraggingMode, 
+				InteractorData.TransformGizmoInteractionType, 
+				bWithTwoHands,
+				*ViewportClient, 
+				InteractorData.OptionalHandlePlacement, 
+				DragDelta, 
+				OtherHandDragDelta,
+				DraggedTo,
+				OtherHandDraggedTo,
+				DragDeltaFromStart, 
+				OtherHandDragDeltaFromStart,
+				LaserPointerStart, 
+				LaserPointerDirection, 
+				bIsLaserPointerValid, 
+				InteractorData.GizmoStartTransform, 
+				InteractorData.GizmoStartLocalBounds,
+				InteractorData.DraggingTransformGizmoComponent.Get(),
+				/* In/Out */ InteractorData.GizmoSpaceFirstDragUpdateOffsetAlongAxis,
+				/* In/Out */ InteractorData.GizmoSpaceDragDeltaFromStartOffset,
+				InteractorData.bIsDrivingVelocityOfSimulatedTransformables,
+				/* Out */ UnsnappedDraggedTo );
+
+
+			// Make sure the hover point is right on the position that we're dragging the object to.  This is important
+			// when constraining dragged objects to a single axis or a plane
+			if( InteractorData.DraggingMode == EViewportInteractionDraggingMode::ActorsWithGizmo )
+			{
+				InteractorData.HoverLocation = FMath::ClosestPointOnSegment( UnsnappedDraggedTo, LaserPointerStart, LaserPointerEnd );
+			}
+
+			if( OtherInteractorThatWasAssistingDrag != nullptr )
+			{
+				OtherInteractorThatWasAssistingDrag->GetInteractorData().LastDragToLocation = OtherHandDraggedTo;
+
+				// Apply damping
+				const bool bVelocitySensitive = InteractorData.DraggingMode == EViewportInteractionDraggingMode::World;
+				ApplyVelocityDamping( OtherInteractorThatWasAssistingDrag->GetInteractorData().DragTranslationVelocity, bVelocitySensitive );
+			}
+
+			// If we were dragging the world, then play some haptic feedback
+			if( InteractorData.DraggingMode == EViewportInteractionDraggingMode::World )
+			{
+				const FVector NewViewLocation = ViewportClient->GetViewLocation();
+
+				// @todo vreditor: Consider doing this for inertial moves too (we need to remember the last hand that invoked the move.)
+
+				const float RoomSpaceHapticTranslationInterval = 25.0f;		// @todo vreditor tweak: Hard coded haptic distance
+				const float WorldSpaceHapticTranslationInterval = RoomSpaceHapticTranslationInterval * WorldScaleFactor;
+
+				bool bCrossedGridLine = false;
+				for( int32 AxisIndex = 0; AxisIndex < 3; ++AxisIndex )
+				{
+					const int32 OldGridCellIndex = FMath::TruncToInt( OldViewLocation[ AxisIndex ] / WorldSpaceHapticTranslationInterval );
+					const int32 NewGridCellIndex = FMath::TruncToInt( NewViewLocation[ AxisIndex ] / WorldSpaceHapticTranslationInterval );
+					if( OldGridCellIndex != NewGridCellIndex )
+					{
+						bCrossedGridLine = true;
+					}
+				}
+
+				if( bCrossedGridLine )
+				{				
+					// @todo vreditor: Make this a velocity-sensitive strength?
+					const float ForceFeedbackStrength = VI::GridHapticFeedbackStrength->GetFloat();		// @todo vreditor tweak: Force feedback strength and enable/disable should be user configurable in options
+					Interactor->PlayHapticEffect( ForceFeedbackStrength );
+					if ( bWithTwoHands )
+					{
+						Interactor->GetOtherInteractor()->PlayHapticEffect( ForceFeedbackStrength );
+					}
+				}
+			}
+		}
+		else if ( InteractorData.DraggingMode == EViewportInteractionDraggingMode::Interactable )
+		{
+			if ( DraggedInteractable )
+			{
+				UViewportDragOperationComponent* DragOperationComponent = DraggedInteractable->GetDragOperationComponent();
+				if ( DragOperationComponent )
+				{
+					UViewportDragOperation* DragOperation = DragOperationComponent->GetDragOperation();
+					if ( DragOperation )
+					{
+						DragOperation->ExecuteDrag( Interactor, DraggedInteractable );
+					}
+				}
+			}
+			else
+			{
+				InteractorData.DraggingMode = EViewportInteractionDraggingMode::Nothing;
+			}
+		}
+		// If we're not actively dragging, apply inertia to any selected elements that we've dragged around recently
+		else 
+		{
+			if( !InteractorData.DragTranslationVelocity.IsNearlyZero( KINDA_SMALL_NUMBER ) &&
+				!InteractorData.bWasAssistingDrag && 	// If we were only assisting, let the other hand take care of doing the update
+				!InteractorData.bIsDrivingVelocityOfSimulatedTransformables )	// If simulation mode is on, let the physics engine take care of inertia
+			{
+				const FVector DragDelta = InteractorData.DragTranslationVelocity;
+				const FVector DraggedTo = InteractorData.LastDragToLocation + DragDelta;
+				const FVector DragDeltaFromStart = DraggedTo - InteractorData.LaserPointerImpactAtDragStart;
+
+				UViewportInteractor* OtherInteractorThatWasAssistingDrag = GetOtherInteractorIntertiaContribute( Interactor );
+				FViewportInteractorData& OtherInteractorThatWasAssistingDragData = OtherInteractorThatWasAssistingDrag->GetInteractorData();
+
+				const bool bWithTwoHands = ( OtherInteractorThatWasAssistingDrag != nullptr );
+
+				FVector OtherHandDragDelta = FVector::ZeroVector;
+				FVector OtherHandDragDeltaFromStart = FVector::ZeroVector;
+				FVector OtherHandDraggedTo = FVector::ZeroVector;
+				if( bWithTwoHands )
+				{
+					OtherHandDragDelta = OtherInteractorThatWasAssistingDragData.DragTranslationVelocity;
+					OtherHandDraggedTo = OtherInteractorThatWasAssistingDragData.LastDragToLocation + OtherHandDragDelta;
+					OtherHandDragDeltaFromStart = OtherHandDraggedTo - OtherInteractorThatWasAssistingDragData.LaserPointerImpactAtDragStart;
+				}
+
+				const bool bIsLaserPointerValid = false;	// Laser pointer has moved on to other things
+				FVector UnsnappedDraggedTo = FVector::ZeroVector;
+				UpdateDragging(
+					DeltaTime,
+					InteractorData.bIsFirstDragUpdate, 
+					InteractorData.LastDraggingMode, 
+					InteractorData.TransformGizmoInteractionType, 
+					bWithTwoHands, 
+					*ViewportClient, 
+					InteractorData.OptionalHandlePlacement, 
+					DragDelta, 
+					OtherHandDragDelta,
+					DraggedTo,
+					OtherHandDraggedTo,
+					DragDeltaFromStart, 
+					OtherHandDragDeltaFromStart,
+					FVector::ZeroVector,	// No valid laser pointer during inertia
+					FVector::ZeroVector,	// No valid laser pointer during inertia
+					bIsLaserPointerValid, 
+					InteractorData.GizmoStartTransform, 
+					InteractorData.GizmoStartLocalBounds, 
+					InteractorData.DraggingTransformGizmoComponent.Get(),
+					/* In/Out */ InteractorData.GizmoSpaceFirstDragUpdateOffsetAlongAxis,
+					/* In/Out */ InteractorData.GizmoSpaceDragDeltaFromStartOffset,
+					InteractorData.bIsDrivingVelocityOfSimulatedTransformables,
+					/* Out */ UnsnappedDraggedTo );
+
+				InteractorData.LastDragToLocation = DraggedTo;
+				const bool bVelocitySensitive = InteractorData.LastDraggingMode == EViewportInteractionDraggingMode::World;
+				ApplyVelocityDamping( InteractorData.DragTranslationVelocity, bVelocitySensitive );
+
+				if( OtherInteractorThatWasAssistingDrag != nullptr )
+				{
+					OtherInteractorThatWasAssistingDragData.LastDragToLocation = OtherHandDraggedTo;
+					ApplyVelocityDamping( OtherInteractorThatWasAssistingDragData.DragTranslationVelocity, bVelocitySensitive );
+				}
+			}
+			else
+			{
+				InteractorData.DragTranslationVelocity = FVector::ZeroVector;
+			}
+		}
+	}
+
+	// Refresh the transform gizmo every frame, just in case actors were moved by some external
+	// influence.  Also, some features of the transform gizmo respond to camera position (like the
+	// measurement text labels), so it just makes sense to keep it up to date.
+	{
+		const bool bNewObjectsSelected = false;
+		const bool bAllHandlesVisible = ( !DraggingActorsWithInteractor.IsSet() );
+		UActorComponent* SingleVisibleHandle = ( DraggingActorsWithInteractor.IsSet() ) ? DraggingActorsWithInteractor.GetValue()->GetInteractorData().DraggingTransformGizmoComponent.Get() : nullptr;
+
+		static TArray< UActorComponent* > HoveringOverHandles;
+		HoveringOverHandles.Reset();
+		for( UViewportInteractor* Interactor : Interactors )
+		{
+			FViewportInteractorData& InteractorData = Interactor->GetInteractorData();
+			UActorComponent* HoveringOverHandle = InteractorData.HoveringOverTransformGizmoComponent.Get();
+			if( HoveringOverHandle != nullptr )
+			{
+				HoveringOverHandles.Add( HoveringOverHandle );
+			}
+		}
+
+		RefreshTransformGizmo( bNewObjectsSelected, bAllHandlesVisible, SingleVisibleHandle, HoveringOverHandles );
+	}
+
+
+	// Update transformables
+	//UpdateTransformables( DeltaTime );
+	const bool bSmoothSnappingEnabled = IsSmoothSnappingEnabled();
+	if ( bSmoothSnappingEnabled || bIsInterpolatingTransformablesFromSnapshotTransform )
+	{
+		const float SmoothSnapSpeed = VI::SmoothSnapSpeed->GetFloat();
+		const bool bUseElasticSnapping = bSmoothSnappingEnabled && VI::ElasticSnap->GetInt() > 0 && TrackingTransaction.IsActive();	// Only while we're still dragging stuff!
+		const float ElasticSnapStrength = VI::ElasticSnapStrength->GetFloat();
+
+		// NOTE: We never sweep while smooth-snapping as the object will never end up where we want it to
+		// due to friction with the ground and other objects.
+		const bool bSweep = false;
+
+		for ( TUniquePtr<FViewportTransformable>& TransformablePtr : Transformables )
+		{
+			FViewportTransformable& Transformable = *TransformablePtr;
+
+			FTransform ActualTargetTransform = Transformable.TargetTransform;
+
+			// If 'elastic snapping' is turned on, we'll have the object 'reach' toward its unsnapped position, so
+			// that the user always has visual feedback when they are dragging something around, even if they
+			// haven't dragged far enough to overcome the snap threshold.
+			if ( bUseElasticSnapping )
+			{
+				ActualTargetTransform.BlendWith( Transformable.UnsnappedTargetTransform, ElasticSnapStrength );
+			}
+
+			FTransform InterpolatedTransform = ActualTargetTransform;
+			if ( !ActualTargetTransform.Equals( Transformable.LastTransform, 0.0f ) )
+			{
+				// If we're really close, just snap it.
+				if ( ActualTargetTransform.Equals( Transformable.LastTransform, KINDA_SMALL_NUMBER ) )
+				{
+					InterpolatedTransform = Transformable.LastTransform = Transformable.UnsnappedTargetTransform = ActualTargetTransform;
+				}
+				else
+				{
+					InterpolatedTransform.Blend(
+						Transformable.LastTransform,
+						ActualTargetTransform,
+						FMath::Min( 1.0f, SmoothSnapSpeed * FMath::Min( DeltaTime, 1.0f / 30.0f ) ) );
+				}
+				Transformable.LastTransform = InterpolatedTransform;
+			}
+
+			if ( bIsInterpolatingTransformablesFromSnapshotTransform )
+			{
+				const float InterpProgress = FMath::Clamp( ( float ) ( CurrentTime - TransformablesInterpolationStartTime ).GetTotalSeconds() / TransformablesInterpolationDuration, 0.0f, 1.0f );
+				InterpolatedTransform.BlendWith( Transformable.InterpolationSnapshotTransform, 1.0f - InterpProgress );
+				if ( InterpProgress >= 1.0f - KINDA_SMALL_NUMBER )
+				{
+					// Finished interpolating
+					bIsInterpolatingTransformablesFromSnapshotTransform = false;
+				}
+			}
+
+			// Got an actor?
+			AActor* Actor = Cast<AActor>( Transformable.Object.Get() );
+			if ( Actor != nullptr )
+			{
+				const FTransform& ExistingTransform = Actor->GetActorTransform();
+				if ( !ExistingTransform.Equals( InterpolatedTransform, 0.0f ) )
+				{
+					const bool bOnlyTranslationChanged =
+						ExistingTransform.GetRotation() == InterpolatedTransform.GetRotation() &&
+						ExistingTransform.GetScale3D() == InterpolatedTransform.GetScale3D();
+
+					Actor->SetActorTransform( InterpolatedTransform, bSweep );
+					//GWarn->Logf( TEXT( "SMOOTH: Actor %s to %s" ), *Actor->GetName(), *Transformable.TargetTransform.ToString() );
+
+					Actor->InvalidateLightingCacheDetailed( bOnlyTranslationChanged );
+
+					const bool bFinished = false;
+					Actor->PostEditMove( bFinished );
+				}
+			}
+			else
+			{
+				// Some other object that we don't know how to drag
+			}
+		}
+	}
+
+	LastWorldToMetersScale = WorldToMetersScale;
+} 
+
+void UViewportWorldInteraction::UpdateTransformables( const float DeltaTime )
+{
+	const FTimespan CurrentTime = FTimespan::FromSeconds( FPlatformTime::Seconds() );
+
+	const bool bSmoothSnappingEnabled = IsSmoothSnappingEnabled();
+	if ( bSmoothSnappingEnabled || bIsInterpolatingTransformablesFromSnapshotTransform )
+	{
+		const float SmoothSnapSpeed = VI::SmoothSnapSpeed->GetFloat();
+		const bool bUseElasticSnapping = bSmoothSnappingEnabled && VI::ElasticSnap->GetInt() > 0 && TrackingTransaction.IsActive();	// Only while we're still dragging stuff!
+		const float ElasticSnapStrength = VI::ElasticSnapStrength->GetFloat();
+
+		// NOTE: We never sweep while smooth-snapping as the object will never end up where we want it to
+		// due to friction with the ground and other objects.
+		const bool bSweep = false;
+
+		for ( TUniquePtr<FViewportTransformable>& TransformablePtr : Transformables )
+		{
+			FViewportTransformable& Transformable = *TransformablePtr;
+
+			FTransform ActualTargetTransform = Transformable.TargetTransform;
+
+			// If 'elastic snapping' is turned on, we'll have the object 'reach' toward its unsnapped position, so
+			// that the user always has visual feedback when they are dragging something around, even if they
+			// haven't dragged far enough to overcome the snap threshold.
+			if ( bUseElasticSnapping )
+			{
+				ActualTargetTransform.BlendWith( Transformable.UnsnappedTargetTransform, ElasticSnapStrength );
+			}
+
+			FTransform InterpolatedTransform = ActualTargetTransform;
+			if ( !ActualTargetTransform.Equals( Transformable.LastTransform, 0.0f ) )
+			{
+				// If we're really close, just snap it.
+				if ( ActualTargetTransform.Equals( Transformable.LastTransform, KINDA_SMALL_NUMBER ) )
+				{
+					InterpolatedTransform = Transformable.LastTransform = Transformable.UnsnappedTargetTransform = ActualTargetTransform;
+				}
+				else
+				{
+					InterpolatedTransform.Blend(
+						Transformable.LastTransform,
+						ActualTargetTransform,
+						FMath::Min( 1.0f, SmoothSnapSpeed * FMath::Min( DeltaTime, 1.0f / 30.0f ) ) );
+				}
+				Transformable.LastTransform = InterpolatedTransform;
+			}
+
+			if ( bIsInterpolatingTransformablesFromSnapshotTransform )
+			{
+				const float InterpProgress = FMath::Clamp( ( float ) ( CurrentTime - TransformablesInterpolationStartTime ).GetTotalSeconds() / TransformablesInterpolationDuration, 0.0f, 1.0f );
+				InterpolatedTransform.BlendWith( Transformable.InterpolationSnapshotTransform, 1.0f - InterpProgress );
+				if ( InterpProgress >= 1.0f - KINDA_SMALL_NUMBER )
+				{
+					// Finished interpolating
+					bIsInterpolatingTransformablesFromSnapshotTransform = false;
+				}
+			}
+
+			// Got an actor?
+			AActor* Actor = Cast<AActor>( Transformable.Object.Get() );
+			if ( Actor != nullptr )
+			{
+				const FTransform& ExistingTransform = Actor->GetActorTransform();
+				if ( !ExistingTransform.Equals( InterpolatedTransform, 0.0f ) )
+				{
+					const bool bOnlyTranslationChanged =
+						ExistingTransform.GetRotation() == InterpolatedTransform.GetRotation() &&
+						ExistingTransform.GetScale3D() == InterpolatedTransform.GetScale3D();
+
+					Actor->SetActorTransform( InterpolatedTransform, bSweep );
+					Actor->InvalidateLightingCacheDetailed( bOnlyTranslationChanged );
+
+					const bool bFinished = false;
+					Actor->PostEditMove( bFinished );
+				}
+			}
+			else
+			{
+				// Some other object that we don't know how to drag
+			}
+		}
+	}
+}
+
+void UViewportWorldInteraction::RefreshOnSelectionChanged( const bool bNewObjectsSelected, bool bAllHandlesVisible, class UActorComponent* SingleVisibleHandle, const TArray< UActorComponent* >& HoveringOverHandles )
+{
+	//@todo: ViewportInteraction: Call listeners for this event
+
+	RefreshTransformGizmo( bNewObjectsSelected, bAllHandlesVisible, SingleVisibleHandle, HoveringOverHandles );
+}
+
+bool UViewportWorldInteraction::IsInteractableComponent( const UActorComponent* Component ) const
+{
+	return true;
+}
+
+ABaseTransformGizmo* UViewportWorldInteraction::GetTransformGizmoActor()
+{
+	SpawnTransformGizmoIfNeeded();
+	return TransformGizmoActor;
+}
+
+void UViewportWorldInteraction::SetDraggedSinceLastSelection( const bool bNewDraggedSinceLastSelection )
+{
+	bDraggedSinceLastSelection = bNewDraggedSinceLastSelection;
+}
+
+void UViewportWorldInteraction::SetLastDragGizmoStartTransform( const FTransform NewLastDragGizmoStartTransform )
+{
+	LastDragGizmoStartTransform = NewLastDragGizmoStartTransform;
+}
+
+TArray<UViewportInteractor*>& UViewportWorldInteraction::GetInteractors()
+{
+	return Interactors;
+}
+
+void UViewportWorldInteraction::UpdateDragging(
+	const float DeltaTime,
+	bool& bIsFirstDragUpdate,
+	const EViewportInteractionDraggingMode DraggingMode,
+	const ETransformGizmoInteractionType InteractionType,
+	const bool bWithTwoHands,
+	FEditorViewportClient& ViewportClient,
+	const TOptional<FTransformGizmoHandlePlacement> OptionalHandlePlacement,
+	const FVector& DragDelta,
+	const FVector& OtherHandDragDelta,
+	const FVector& DraggedTo,
+	const FVector& OtherHandDraggedTo,
+	const FVector& DragDeltaFromStart,
+	const FVector& OtherHandDragDeltaFromStart,
+	const FVector& LaserPointerStart,
+	const FVector& LaserPointerDirection,
+	const bool bIsLaserPointerValid,
+	const FTransform& GizmoStartTransform,
+	const FBox& GizmoStartLocalBounds,
+	const USceneComponent* const DraggingTransformGizmoComponent,
+	FVector& GizmoSpaceFirstDragUpdateOffsetAlongAxis,
+	FVector& DragDeltaFromStartOffset,
+	bool& bIsDrivingVelocityOfSimulatedTransformables,
+	FVector& OutUnsnappedDraggedTo )
+{
+	// IMPORTANT NOTE: Depending on the DraggingMode, the incoming coordinates can be in different coordinate spaces.
+	//		-> When dragging objects around, everything is in world space unless otherwise specified in the variable name.
+	//		-> When dragging the world around, everything is in room space unless otherwise specified.
+
+	bool bMovedAnyTransformables = false;
+
+	// This will be set to true if we want to set the physics velocity on the dragged objects based on how far
+	// we dragged this frame.  Used for pushing objects around in Simulate mode.  If this is set to false, we'll
+	// simply zero out the velocities instead to cancel out the physics engine's applied forces (e.g. gravity) 
+	// while we're dragging objects around (otherwise, the objects will spaz out as soon as dragging stops.)
+	bool bShouldApplyVelocitiesFromDrag = false;
+
+	// Always snap objects relative to where they were when we first grabbed them.  We never want objects to immediately
+	// teleport to the closest snap, but instead snaps should always be relative to the gizmo center
+	// NOTE: While placing objects, we always snap in world space, since the initial position isn't really at all useful
+	const bool bLocalSpaceSnapping =
+		GetTransformGizmoCoordinateSpace() == COORD_Local &&
+		DraggingMode != EViewportInteractionDraggingMode::ActorsAtLaserImpact;
+	const FVector SnapGridBase = ( DraggingMode == EViewportInteractionDraggingMode::ActorsAtLaserImpact||  bLocalSpaceSnapping ) ? FVector::ZeroVector : GizmoStartTransform.GetLocation(); //@todo ViewportInteraction
+
+
+	// Okay, time to move stuff!  We'll do this part differently depending on whether we're dragging actual actors around
+	// or we're moving the camera (aka. dragging the world)
+	if ( DraggingMode == EViewportInteractionDraggingMode::ActorsWithGizmo ||
+		DraggingMode == EViewportInteractionDraggingMode::ActorsAtLaserImpact )
+	{
+		FVector ConstrainedDragDeltaFromStart = ComputeConstrainedDragDeltaFromStart(
+			bIsFirstDragUpdate,
+			InteractionType,
+			OptionalHandlePlacement,
+			DragDeltaFromStart,
+			LaserPointerStart,
+			LaserPointerDirection,
+			bIsLaserPointerValid,
+			GizmoStartTransform,
+			/* In/Out */ GizmoSpaceFirstDragUpdateOffsetAlongAxis,
+			/* In/Out */ DragDeltaFromStartOffset );
+
+		FVector UnsnappedDraggedTo = GizmoStartTransform.GetLocation() + ConstrainedDragDeltaFromStart;
+		OutUnsnappedDraggedTo = UnsnappedDraggedTo;
+
+		// Grid snap!
+		FVector SnappedDraggedTo = UnsnappedDraggedTo;
+		if ( FSnappingUtils::IsSnapToGridEnabled() )
+		{
+			// Snap in local space, if we need to
+			if ( bLocalSpaceSnapping )
+			{
+				SnappedDraggedTo = GizmoStartTransform.InverseTransformPositionNoScale( SnappedDraggedTo );
+				FSnappingUtils::SnapPointToGrid( SnappedDraggedTo, SnapGridBase );
+				SnappedDraggedTo = GizmoStartTransform.TransformPositionNoScale( SnappedDraggedTo );
+			}
+			else
+			{
+				FSnappingUtils::SnapPointToGrid( SnappedDraggedTo, SnapGridBase );
+			}
+		}
+
+		// Two passes.  First update the real transform.  Then update the unsnapped transform.
+		for ( int32 PassIndex = 0; PassIndex < 2; ++PassIndex )
+		{
+			const bool bIsUpdatingUnsnappedTarget = ( PassIndex == 1 );
+			const FVector& PassDraggedTo = bIsUpdatingUnsnappedTarget ? UnsnappedDraggedTo : SnappedDraggedTo;
+
+			if ( InteractionType == ETransformGizmoInteractionType::Translate ||
+				InteractionType == ETransformGizmoInteractionType::TranslateOnPlane )
+			{
+				// Translate the objects!
+				for ( TUniquePtr<FViewportTransformable>& TransformablePtr : Transformables )
+				{
+					FViewportTransformable& Transformable = *TransformablePtr;
+					FTransform& PassTargetTransform = bIsUpdatingUnsnappedTarget ? Transformable.UnsnappedTargetTransform : Transformable.TargetTransform;
+
+					const FVector OldLocation = PassTargetTransform.GetLocation();
+					PassTargetTransform.SetLocation( ( Transformable.StartTransform.GetLocation() - GizmoStartTransform.GetLocation() ) + PassDraggedTo );
+
+					bMovedAnyTransformables = true;
+					bShouldApplyVelocitiesFromDrag = true;
+				}
+			}
+			else if ( InteractionType == ETransformGizmoInteractionType::StretchAndReposition )
+			{
+				// We only support stretching by a handle currently
+				const FTransformGizmoHandlePlacement& HandlePlacement = OptionalHandlePlacement.GetValue();
+
+				const FVector PassGizmoSpaceDraggedTo = GizmoStartTransform.InverseTransformPositionNoScale( PassDraggedTo );
+
+				FBox NewGizmoLocalBounds = GizmoStartLocalBounds;
+				FVector GizmoSpacePivotLocation = FVector::ZeroVector;
+				for ( int32 AxisIndex = 0; AxisIndex < 3; ++AxisIndex )
+				{
+					// Figure out how much the gizmo bounds changes
+					if ( HandlePlacement.Axes[ AxisIndex ] == ETransformGizmoHandleDirection::Negative )	// Negative direction
+					{
+						GizmoSpacePivotLocation[ AxisIndex ] = GizmoStartLocalBounds.Max[ AxisIndex ];
+						NewGizmoLocalBounds.Min[ AxisIndex ] = GizmoStartLocalBounds.Min[ AxisIndex ] + PassGizmoSpaceDraggedTo[ AxisIndex ];
+
+					}
+					else if ( HandlePlacement.Axes[ AxisIndex ] == ETransformGizmoHandleDirection::Positive )	// Positive direction
+					{
+						GizmoSpacePivotLocation[ AxisIndex ] = GizmoStartLocalBounds.Min[ AxisIndex ];
+						NewGizmoLocalBounds.Max[ AxisIndex ] = GizmoStartLocalBounds.Max[ AxisIndex ] + PassGizmoSpaceDraggedTo[ AxisIndex ];
+					}
+					else
+					{
+						// Must be ETransformGizmoHandleDirection::Center.  
+						GizmoSpacePivotLocation[ AxisIndex ] = GizmoStartLocalBounds.GetCenter()[ AxisIndex ];
+					}
+				}
+				
+				const FVector GizmoStartLocalSize = GizmoStartLocalBounds.GetSize();
+				const FVector NewGizmoLocalSize = NewGizmoLocalBounds.GetSize();
+
+				FVector NewGizmoLocalScaleFromStart = FVector( 1.0f );
+				if ( !FMath::IsNearlyZero( NewGizmoLocalSize.GetAbsMin() ) )
+				{
+					NewGizmoLocalScaleFromStart = NewGizmoLocalSize / GizmoStartLocalSize;
+				}
+				else
+				{
+					// Zero scale.  This is allowed in Unreal, for better or worse.
+					NewGizmoLocalScaleFromStart = FVector( 0 );
+				}
+
+				// Stretch and reposition the objects!
+				for ( TUniquePtr<FViewportTransformable>& TransformablePtr : Transformables )
+				{
+					FViewportTransformable& Transformable = *TransformablePtr;
+
+					FTransform& PassTargetTransform = bIsUpdatingUnsnappedTarget ? Transformable.UnsnappedTargetTransform : Transformable.TargetTransform;
+
+					{
+						const FVector GizmoSpaceTransformableStartLocation = GizmoStartTransform.InverseTransformPositionNoScale( Transformable.StartTransform.GetLocation() );
+						const FVector NewGizmoSpaceLocation = ( GizmoSpaceTransformableStartLocation - GizmoSpacePivotLocation ) * NewGizmoLocalScaleFromStart + GizmoSpacePivotLocation;
+						const FVector OldLocation = PassTargetTransform.GetLocation();
+						PassTargetTransform.SetLocation( GizmoStartTransform.TransformPosition( NewGizmoSpaceLocation ) );
+					}
+
+					// @todo vreditor: This scale is still in gizmo space, but we're setting it in world space
+					PassTargetTransform.SetScale3D( Transformable.StartTransform.GetScale3D() * NewGizmoLocalScaleFromStart );
+
+					bMovedAnyTransformables = true;
+					bShouldApplyVelocitiesFromDrag = false;
+				}
+			}
+			else if ( InteractionType == ETransformGizmoInteractionType::Scale || InteractionType == ETransformGizmoInteractionType::UniformScale )
+			{
+				const FTransformGizmoHandlePlacement& HandlePlacement = OptionalHandlePlacement.GetValue();
+
+				int32 CenterHandleCount, FacingAxisIndex, CenterAxisIndex;
+				HandlePlacement.GetCenterHandleCountAndFacingAxisIndex( /* Out */ CenterHandleCount, /* Out */ FacingAxisIndex, /* Out */ CenterAxisIndex );
+
+				const FVector PassGizmoSpaceDraggedTo = GizmoStartTransform.InverseTransformPositionNoScale( PassDraggedTo );
+
+				float AddedScaleOnAxis = 0.0f;
+				if ( InteractionType == ETransformGizmoInteractionType::Scale )
+				{
+					AddedScaleOnAxis = PassGizmoSpaceDraggedTo[ FacingAxisIndex ] * VI::ScaleSensitivity->GetFloat();
+
+					// Invert if we we are scaling on the negative side of the gizmo
+					if ( DraggingTransformGizmoComponent && DraggingTransformGizmoComponent->GetRelativeTransform().GetLocation()[ FacingAxisIndex ] < 0 )
+					{
+						AddedScaleOnAxis *= -1;
+					}
+				}
+				else if ( InteractionType == ETransformGizmoInteractionType::UniformScale )
+				{
+					// Always use Z for uniform scale
+					const FVector RelativeDraggedTo = PassDraggedTo - GizmoStartTransform.GetLocation();
+					AddedScaleOnAxis = RelativeDraggedTo.Z * VI::ScaleSensitivity->GetFloat();
+				}
+
+				// Scale the objects!
+				for ( TUniquePtr<FViewportTransformable>& TransformablePtr : Transformables )
+				{
+					FViewportTransformable& Transformable = *TransformablePtr;
+					FTransform& PassTargetTransform = bIsUpdatingUnsnappedTarget ? Transformable.UnsnappedTargetTransform : Transformable.TargetTransform;
+					FVector NewTotalScale = Transformable.StartTransform.GetScale3D();
+
+					if ( InteractionType == ETransformGizmoInteractionType::Scale )
+					{
+						NewTotalScale[ FacingAxisIndex ] += AddedScaleOnAxis;
+					}
+					else if ( InteractionType == ETransformGizmoInteractionType::UniformScale )
+					{
+						NewTotalScale += FVector( AddedScaleOnAxis );
+					}
+
+					// Scale snap!
+					if ( !bIsUpdatingUnsnappedTarget && FSnappingUtils::IsScaleSnapEnabled() )
+					{
+						FSnappingUtils::SnapScale( NewTotalScale, FVector::ZeroVector );
+					}
+
+					PassTargetTransform.SetScale3D( NewTotalScale );
+					bMovedAnyTransformables = true;
+					bShouldApplyVelocitiesFromDrag = true;
+				}
+			}
+			else if ( InteractionType == ETransformGizmoInteractionType::Rotate )
+			{
+				// We only support rotating by a handle currently
+				const FTransformGizmoHandlePlacement& HandlePlacement = OptionalHandlePlacement.GetValue();
+
+				int32 CenterHandleCount, FacingAxisIndex, CenterAxisIndex;
+				HandlePlacement.GetCenterHandleCountAndFacingAxisIndex( /* Out */ CenterHandleCount, /* Out */ FacingAxisIndex, /* Out */ CenterAxisIndex );
+				check( CenterAxisIndex != INDEX_NONE );
+
+				// Get the local location of the rotation handle
+				FVector HandleRelativeLocation = FVector::ZeroVector;
+				for ( int32 AxisIndex = 0; AxisIndex < 3; ++AxisIndex )
+				{
+					if ( HandlePlacement.Axes[ AxisIndex ] == ETransformGizmoHandleDirection::Negative )	// Negative direction
+					{
+						HandleRelativeLocation[ AxisIndex ] = GizmoStartLocalBounds.Min[ AxisIndex ];
+					}
+					else if ( HandlePlacement.Axes[ AxisIndex ] == ETransformGizmoHandleDirection::Positive )	// Positive direction
+					{
+						HandleRelativeLocation[ AxisIndex ] = GizmoStartLocalBounds.Max[ AxisIndex ];
+					}
+					else // ETransformGizmoHandleDirection::Center
+					{
+						HandleRelativeLocation[ AxisIndex ] = GizmoStartLocalBounds.GetCenter()[ AxisIndex ];
+					}
+				}
+
+				const FVector GizmoSpaceRotationAxis = ( CenterAxisIndex == 0 ) ? FVector::ForwardVector : ( CenterAxisIndex == 1 ? FVector::RightVector : FVector::UpVector );
+
+				// Make a vector that points along the tangent vector of the bounding box edge
+				const FVector GizmoSpaceTowardHandleVector = ( HandleRelativeLocation - GizmoStartLocalBounds.GetCenter() ).GetSafeNormal();
+				const FVector GizmoSpaceDiagonalAxis = FQuat( GizmoSpaceRotationAxis, FMath::DegreesToRadians( 90.0f ) ).RotateVector( GizmoSpaceTowardHandleVector );
+
+				// Figure out how far we've dragged in the direction of the diagonal axis.
+				const float RotationProgress = FVector::DotProduct( DragDeltaFromStart, GizmoStartTransform.TransformVectorNoScale( GizmoSpaceDiagonalAxis ) );
+
+				const float RotationDegreesAroundAxis = RotationProgress * VI::GizmoRotationSensitivity->GetFloat();
+				const FQuat RotationDeltaFromStart = FQuat( GizmoSpaceRotationAxis, FMath::DegreesToRadians( RotationDegreesAroundAxis ) );
+
+				const FTransform WorldToGizmo = GizmoStartTransform.Inverse();
+
+				// Rotate (and reposition) the objects!
+				for ( TUniquePtr<FViewportTransformable>& TransformablePtr : Transformables )
+				{
+					FViewportTransformable& Transformable = *TransformablePtr;
+					FTransform& PassTargetTransform = bIsUpdatingUnsnappedTarget ? Transformable.UnsnappedTargetTransform : Transformable.TargetTransform;
+
+					const FTransform GizmoSpaceStartTransform = Transformable.StartTransform * WorldToGizmo;
+					FTransform GizmoSpaceRotatedTransform = GizmoSpaceStartTransform * RotationDeltaFromStart;
+
+					// Snap rotation in gizmo space
+					if ( !bIsUpdatingUnsnappedTarget )
+					{
+						FRotator SnappedRotation = GizmoSpaceRotatedTransform.GetRotation().Rotator();
+						FSnappingUtils::SnapRotatorToGrid( SnappedRotation );
+						GizmoSpaceRotatedTransform.SetRotation( SnappedRotation.Quaternion() );
+					}
+
+					const FTransform WorldSpaceRotatedTransform = GizmoSpaceRotatedTransform * GizmoStartTransform;
+
+					PassTargetTransform = WorldSpaceRotatedTransform;
+					bMovedAnyTransformables = true;
+					bShouldApplyVelocitiesFromDrag = false;
+
+					// Snap location in world space
+					if ( !bIsUpdatingUnsnappedTarget )
+					{
+						FVector SnappedLocation = PassTargetTransform.GetLocation();
+
+						// Snap in local space, if we need to
+						if ( bLocalSpaceSnapping )
+						{
+							SnappedLocation = GizmoStartTransform.InverseTransformPositionNoScale( SnappedLocation );
+							FSnappingUtils::SnapPointToGrid( SnappedLocation, SnapGridBase );
+							SnappedLocation = GizmoStartTransform.TransformPositionNoScale( SnappedLocation );
+						}
+						else
+						{
+							FSnappingUtils::SnapPointToGrid( SnappedLocation, SnapGridBase );
+						}
+						PassTargetTransform.SetLocation( SnappedLocation );
+					}
+				}
+			}
+			else if ( InteractionType == ETransformGizmoInteractionType::RotateOnAngle )
+			{
+				const FTransformGizmoHandlePlacement& HandlePlacement = OptionalHandlePlacement.GetValue();
+
+				int32 CenterHandleCount, FacingAxisIndex, CenterAxisIndex;
+				HandlePlacement.GetCenterHandleCountAndFacingAxisIndex( /* Out */ CenterHandleCount, /* Out */ FacingAxisIndex, /* Out */ CenterAxisIndex );
+
+				FVector GizmoSpaceFacingAxisVector = UGizmoHandleGroup::GetAxisVector( FacingAxisIndex, HandlePlacement.Axes[ FacingAxisIndex ] );
+
+				if ( DraggingTransformGizmoComponent )
+				{
+					const FTransform WorldToGizmo = GizmoStartTransform.Inverse();
+
+					// Rotate the objects!
+					for ( TUniquePtr<FViewportTransformable>& TransformablePtr : Transformables )
+					{
+						FViewportTransformable& Transformable = *TransformablePtr;
+						FTransform& PassTargetTransform = bIsUpdatingUnsnappedTarget ? Transformable.UnsnappedTargetTransform : Transformable.TargetTransform;
+
+						const FVector RotationHandlePlaneLocation =  Transformable.StartTransform.GetLocation();
+						if ( !StartDragHandleDirection.IsSet() )
+						{
+							StartDragHandleDirection = DraggingTransformGizmoComponent->GetComponentTransform().GetRotation().Vector();
+							StartDragHandleDirection->Normalize();
+						}
+
+						// Get the laserpointer intersection on the plane of the handle
+						FPlane Plane = FPlane( RotationHandlePlaneLocation, StartDragHandleDirection.GetValue() );
+						
+						const FVector IntersectWorldLocationOnPlane = FMath::LinePlaneIntersection( LaserPointerStart, LaserPointerStart + LaserPointerDirection, 
+							RotationHandlePlaneLocation, StartDragHandleDirection.GetValue() );
+
+						DrawDebugLine( GetViewportWorld(), TransformGizmoActor->GetActorLocation(), IntersectWorldLocationOnPlane, FColor::Red, false, -1.f, 0, 3.0f * GetWorldScaleFactor() );
+
+						// The starting transformation of this transformable in gizmo space
+						const FTransform GizmoSpaceStartTransform = Transformable.StartTransform * WorldToGizmo;
+
+						// Relative offset of the intersection on the plane
+						const FVector IntersectLocalLocationOnPlane = Transformable.StartTransform.InverseTransformPositionNoScale( IntersectWorldLocationOnPlane );
+						FVector RotatedIntersectLocationOnPlane;
+						if ( GLevelEditorModeTools().GetCoordSystem( true ) == COORD_Local )
+						{
+							RotatedIntersectLocationOnPlane = GizmoSpaceFacingAxisVector.Rotation().UnrotateVector( IntersectLocalLocationOnPlane );
+						}
+						else
+						{
+							const FVector TransformableSpaceFacingAxisVector = Transformable.StartTransform.InverseTransformVectorNoScale( GizmoStartTransform.InverseTransformVectorNoScale( GizmoSpaceFacingAxisVector ) );
+							RotatedIntersectLocationOnPlane = TransformableSpaceFacingAxisVector.Rotation().UnrotateVector( IntersectLocalLocationOnPlane );
+						}
+
+						// Get the angle between the center and the intersected point
+						float AngleToIntersectedLocation = FMath::Atan2( RotatedIntersectLocationOnPlane.Y, RotatedIntersectLocationOnPlane.Z );
+						if ( !StartDragAngleOnRotation.IsSet() )
+						{
+							StartDragAngleOnRotation = AngleToIntersectedLocation;
+						}
+						
+						// Delta rotation in gizmo space between the starting and the intersection rotation
+						const float AngleDeltaRotationFromStart = FMath::FindDeltaAngleRadians( AngleToIntersectedLocation, StartDragAngleOnRotation.GetValue() );
+						const FQuat GizmoSpaceDeltaRotation = FQuat( GizmoSpaceFacingAxisVector, AngleDeltaRotationFromStart );
+
+						// Add the delta rotation to the starting transformation
+						FTransform GizmoSpaceRotatedTransform = GizmoSpaceStartTransform * GizmoSpaceDeltaRotation;
+
+						// Snap rotation in gizmo space
+						if ( !bIsUpdatingUnsnappedTarget )
+						{
+							FRotator SnappedRotation = GizmoSpaceRotatedTransform.GetRotation().Rotator();
+							FSnappingUtils::SnapRotatorToGrid( SnappedRotation );
+							GizmoSpaceRotatedTransform.SetRotation( SnappedRotation.Quaternion() );
+						}
+
+						// Rotation back to world space 
+						const FTransform WorldSpaceRotatedTransform = GizmoSpaceRotatedTransform * GizmoStartTransform;
+
+						PassTargetTransform = WorldSpaceRotatedTransform;
+						bMovedAnyTransformables = true;
+						bShouldApplyVelocitiesFromDrag = true;
+					}
+				}
+			}
+		}
+	}
+	else if ( DraggingMode == EViewportInteractionDraggingMode::ActorsFreely ||
+		DraggingMode == EViewportInteractionDraggingMode::World )
+	{
+		FVector TranslationOffset = FVector::ZeroVector;
+		FQuat RotationOffset = FQuat::Identity;
+		FVector ScaleOffset = FVector( 1.0f );
+
+		bool bHasPivotLocation = false;
+		FVector PivotLocation = FVector::ZeroVector;
+
+		if ( bWithTwoHands )
+		{
+			// Rotate/scale while simultaneously counter-translating (two hands)
+
+			// NOTE: The reason that we use per-frame deltas (rather than the delta from the start of the whole interaction,
+			// like we do with the other interaction modes), is because the point that we're rotating/scaling relative to
+			// actually changes every update as the user moves their hands.  So it's iteratively getting the user to
+			// the result they naturally expect.
+
+			const FVector LineStart = DraggedTo;
+			const FVector LineStartDelta = DragDelta;
+			const float LineStartDistance = LineStartDelta.Size();
+			const FVector LastLineStart = LineStart - LineStartDelta;
+
+			const FVector LineEnd = OtherHandDraggedTo;
+			const FVector LineEndDelta = OtherHandDragDelta;
+			const float LineEndDistance = LineEndDelta.Size();
+			const FVector LastLineEnd = LineEnd - LineEndDelta;
+
+			// Choose a point along the new line segment to rotate and scale relative to.  We'll weight it toward the
+			// side of the line that moved the most this update.
+			const float TotalDistance = LineStartDistance + LineEndDistance;
+			float LineStartToEndActivityWeight = 0.5f;	// Default to right in the center, if no distance moved yet.
+			if ( !FMath::IsNearlyZero( TotalDistance ) )	// Avoid division by zero
+			{
+				LineStartToEndActivityWeight = LineStartDistance / TotalDistance;
+			}
+
+			PivotLocation = FMath::Lerp( LastLineStart, LastLineEnd, LineStartToEndActivityWeight );
+			bHasPivotLocation = true;
+
+			if ( DraggingMode == EViewportInteractionDraggingMode::World && VI::ScaleWorldFromFloor->GetInt() != 0 )
+			{
+				PivotLocation.Z = 0.0f;
+			}
+
+			// @todo vreditor debug
+			if ( false )
+			{
+				const FTransform LinesRelativeTo = DraggingMode == EViewportInteractionDraggingMode::World ? GetRoomTransform() : FTransform::Identity;
+				DrawDebugLine( GetViewportWorld(), LinesRelativeTo.TransformPosition( LineStart ), LinesRelativeTo.TransformPosition( LineEnd ), FColor::Green, false, 0.0f );
+				DrawDebugLine( GetViewportWorld(), LinesRelativeTo.TransformPosition( LastLineStart ), LinesRelativeTo.TransformPosition( LastLineEnd ), FColor::Red, false, 0.0f );
+
+				DrawDebugSphere( GetViewportWorld(), LinesRelativeTo.TransformPosition( PivotLocation ), 2.5f * GetWorldScaleFactor(), 32, FColor::White, false, 0.0f );
+			}
+
+			const float LastLineLength = ( LastLineEnd - LastLineStart ).Size();
+			const float LineLength = ( LineEnd - LineStart ).Size();
+			ScaleOffset = FVector( LineLength / LastLineLength );
+			//			ScaleOffset = FVector( 0.98f + 0.04f * FMath::MakePulsatingValue( FPlatformTime::Seconds(), 0.1f ) ); // FVector( LineLength / LastLineLength );
+
+			// How much did the line rotate since last time?
+			RotationOffset = FQuat::FindBetweenVectors( LastLineEnd - LastLineStart, LineEnd - LineStart );
+
+			// For translation, only move proportionally to the common vector between the two deltas.  Basically,
+			// you need to move both hands in the same direction to translate while gripping with two hands.
+			const FVector AverageDelta = ( DragDelta + OtherHandDragDelta ) * 0.5f;
+			const float TranslationWeight = FMath::Max( 0.0f, FVector::DotProduct( DragDelta.GetSafeNormal(), OtherHandDragDelta.GetSafeNormal() ) );
+			TranslationOffset = FMath::Lerp( FVector::ZeroVector, AverageDelta, TranslationWeight );
+		}
+		else
+		{
+			// Translate only (one hand)
+			TranslationOffset = DragDelta;
+		}
+
+		if ( VI::AllowVerticalWorldMovement->GetInt() == 0 )
+		{
+			TranslationOffset.Z = 0.0f;
+		}
+
+
+		if ( DraggingMode == EViewportInteractionDraggingMode::ActorsFreely )
+		{
+			for ( TUniquePtr<FViewportTransformable>& TransformablePtr : Transformables )
+			{
+				FViewportTransformable& Transformable = *TransformablePtr;
+
+				const FTransform OldTransform = Transformable.UnsnappedTargetTransform;
+
+				// Two passes.  First update the real transform.  Then update the unsnapped transform.
+				for ( int32 PassIndex = 0; PassIndex < 2; ++PassIndex )
+				{
+					const bool bIsUpdatingUnsnappedTarget = ( PassIndex == 1 );
+
+					// Scale snap!
+					FVector SnappedScaleOffset = ScaleOffset;
+					if ( !bIsUpdatingUnsnappedTarget && FSnappingUtils::IsScaleSnapEnabled() )
+					{
+						// Did scale even change?
+						if ( !( FMath::IsNearlyEqual( ScaleOffset.GetAbsMax(), 1.0f ) && FMath::IsNearlyEqual( ScaleOffset.GetAbsMin(), 1.0f ) ) )
+						{
+							FVector NewTotalScale = OldTransform.GetScale3D() * ScaleOffset;
+							FSnappingUtils::SnapScale( NewTotalScale, FVector::ZeroVector );
+							SnappedScaleOffset = NewTotalScale / OldTransform.GetScale3D();
+						}
+					}
+
+					const FTransform ScaleOffsetTransform( FQuat::Identity, FVector::ZeroVector, SnappedScaleOffset );
+					const FTransform TranslationOffsetTransform( FQuat::Identity, TranslationOffset );
+					const FTransform RotationOffsetTransform( RotationOffset, FVector::ZeroVector );
+
+					if ( !bHasPivotLocation )
+					{
+						PivotLocation = OldTransform.GetLocation();
+					}
+
+					// @todo vreditor multi: This is only solving for rotating/scaling one object.  Really we want to rotate/scale the GIZMO and have all objects rotated and repositioned within it!
+					FTransform NewTransform;
+					const bool bIgnoreGizmoSpace = true;
+					if ( bIgnoreGizmoSpace )
+					{
+						const FTransform PivotToWorld = FTransform( FQuat::Identity, PivotLocation );
+						const FTransform WorldToPivot = FTransform( FQuat::Identity, -PivotLocation );
+						NewTransform = OldTransform * WorldToPivot * ScaleOffsetTransform * RotationOffsetTransform * PivotToWorld * TranslationOffsetTransform;
+					}
+					else
+					{
+						// @todo vreditor perf: Can move most of this outside of the inner loop.  Also could be optimized by reducing the number of transforms.
+						const FTransform GizmoToWorld = GizmoStartTransform;
+						const FTransform WorldToGizmo = GizmoToWorld.Inverse();
+						const FVector GizmoSpacePivotLocation = WorldToGizmo.TransformPositionNoScale( PivotLocation );
+						const FTransform GizmoToPivot = FTransform( FQuat::Identity, -GizmoSpacePivotLocation );
+						const FTransform PivotToGizmo = FTransform( FQuat::Identity, GizmoSpacePivotLocation );
+						const FQuat GizmoSpaceRotationOffsetQuat = WorldToGizmo.GetRotation() * RotationOffset;	// NOTE: FQuat composition order is opposite that of FTransform!
+						const FTransform GizmoSpaceRotationOffsetTransform( GizmoSpaceRotationOffsetQuat, FVector::ZeroVector );
+						NewTransform = OldTransform * TranslationOffsetTransform * WorldToGizmo * GizmoToPivot * ScaleOffsetTransform * GizmoSpaceRotationOffsetTransform * PivotToGizmo * GizmoToWorld;
+					}
+
+					if ( bIsUpdatingUnsnappedTarget )
+					{
+						const FVector OldLocation = Transformable.UnsnappedTargetTransform.GetLocation();
+						Transformable.UnsnappedTargetTransform = NewTransform;
+					}
+					else
+					{
+						Transformable.TargetTransform = NewTransform;
+
+						// Grid snap!
+						FVector SnappedLocation = Transformable.TargetTransform.GetLocation();
+						if ( bLocalSpaceSnapping )
+						{
+							SnappedLocation = GizmoStartTransform.InverseTransformPositionNoScale( SnappedLocation );
+							FSnappingUtils::SnapPointToGrid( SnappedLocation, SnapGridBase );
+							SnappedLocation = GizmoStartTransform.TransformPositionNoScale( SnappedLocation );
+						}
+						else
+						{
+							FSnappingUtils::SnapPointToGrid( SnappedLocation, SnapGridBase );
+						}
+						Transformable.TargetTransform.SetLocation( SnappedLocation );
+
+						// Rotation snap!
+						FRotator SnappedRotation = Transformable.TargetTransform.GetRotation().Rotator();
+						FSnappingUtils::SnapRotatorToGrid( SnappedRotation );
+						Transformable.TargetTransform.SetRotation( SnappedRotation.Quaternion() );
+					}
+				}
+
+				bMovedAnyTransformables = true;
+				bShouldApplyVelocitiesFromDrag = true;
+			}
+		}
+		else if ( DraggingMode == EViewportInteractionDraggingMode::World )
+		{
+			FTransform RoomTransform = GetRoomTransform();
+
+			// Adjust world scale
+			const float WorldScaleOffset = ScaleOffset.GetAbsMax();
+			if ( WorldScaleOffset != 0.0f )
+			{
+				const float OldWorldToMetersScale = GetViewportWorld()->GetWorldSettings()->WorldToMeters;
+				const float NewWorldToMetersScale = OldWorldToMetersScale / WorldScaleOffset;
+
+				// NOTE: Instead of clamping, we simply skip changing the W2M this frame if it's out of bounds.  Clamping makes our math more complicated.
+				if ( NewWorldToMetersScale != OldWorldToMetersScale &&
+					NewWorldToMetersScale >= GetMinScale() && NewWorldToMetersScale <= GetMaxScale() )
+				{
+					SetWorldToMetersScale( NewWorldToMetersScale );
+
+					// Because the tracking space size has changed, but our head position within that space relative to the origin
+					// of the room is the same (before scaling), we need to offset our location within the tracking space to compensate.
+					// This makes the user feel like their head and hands remain in the same location.
+					const FVector RoomSpacePivotLocation = PivotLocation;
+					const FVector WorldSpacePivotLocation = RoomTransform.TransformPosition( RoomSpacePivotLocation );
+
+					const FVector NewRoomSpacePivotLocation = ( RoomSpacePivotLocation / OldWorldToMetersScale ) * NewWorldToMetersScale;
+					const FVector NewWorldSpacePivotLocation = RoomTransform.TransformPosition( NewRoomSpacePivotLocation );
+
+					const FVector WorldSpacePivotDelta = ( NewWorldSpacePivotLocation - WorldSpacePivotLocation );
+
+					const FVector NewWorldSpaceRoomLocation = RoomTransform.GetLocation() - WorldSpacePivotDelta;
+
+					RoomTransform.SetLocation( NewWorldSpaceRoomLocation );
+				}
+			}
+
+			// Apply rotation and translation
+			{
+				FTransform RotationOffsetTransform = FTransform::Identity;
+				RotationOffsetTransform.SetRotation( RotationOffset );
+
+				if ( VI::AllowWorldRotationPitchAndRoll->GetInt() == 0 )
+				{
+					// Eliminate pitch and roll in rotation offset.  We don't want the user to get sick!
+					FRotator YawRotationOffset = RotationOffset.Rotator();
+					YawRotationOffset.Pitch = YawRotationOffset.Roll = 0.0f;
+					RotationOffsetTransform.SetRotation( YawRotationOffset.Quaternion() );
+				}
+
+				// Move the camera in the opposite direction, so it feels to the user as if they're dragging the entire world around
+				const FTransform TranslationOffsetTransform( FQuat::Identity, TranslationOffset );
+				const FTransform PivotToWorld = FTransform( FQuat::Identity, PivotLocation ) * RoomTransform;
+				const FTransform WorldToPivot = PivotToWorld.Inverse();
+				RoomTransform = TranslationOffsetTransform.Inverse() * RoomTransform * WorldToPivot * RotationOffsetTransform.Inverse() * PivotToWorld;
+			}
+
+			SetRoomTransform( RoomTransform );
+		}
+
+	}
+
+
+	// If we're not using smooth snapping, go ahead and immediately update the transforms of all objects.  If smooth
+	// snapping is enabled, this will be done in Tick() instead.
+	if ( bMovedAnyTransformables )
+	{
+		// Update velocity if we're simulating in editor
+		if ( GEditor->bIsSimulatingInEditor )
+		{
+			for ( TUniquePtr<FViewportTransformable>& TransformablePtr : Transformables )
+			{
+				FViewportTransformable& Transformable = *TransformablePtr;
+
+				AActor* Actor = Cast<AActor>( Transformable.Object.Get() );
+				if ( Actor != nullptr )
+				{
+					UPrimitiveComponent* RootPrimitiveComponent = Cast<UPrimitiveComponent>( Actor->GetRootComponent() );
+					if ( RootPrimitiveComponent != nullptr )
+					{
+						// @todo VREditor physics: When freely transforming, should set angular velocity too (will pivot be the same though??)
+						if ( RootPrimitiveComponent->IsSimulatingPhysics() )
+						{
+							FVector MoveDelta = Transformable.UnsnappedTargetTransform.GetLocation() - Transformable.LastTransform.GetLocation();
+							if ( bShouldApplyVelocitiesFromDrag )
+							{
+								bIsDrivingVelocityOfSimulatedTransformables = true;
+							}
+							else
+							{
+								MoveDelta = FVector::ZeroVector;
+							}
+							RootPrimitiveComponent->SetAllPhysicsLinearVelocity( MoveDelta * DeltaTime * VI::InertiaVelocityBoost->GetFloat() );
+						}
+					}
+				}
+			}
+		}
+
+		const bool bSmoothSnappingEnabled = IsSmoothSnappingEnabled();
+		if ( !bSmoothSnappingEnabled && !bIsInterpolatingTransformablesFromSnapshotTransform )
+		{
+			const bool bSweep = bShouldApplyVelocitiesFromDrag && VI::SweepPhysicsWhileSimulating->GetInt() != 0 && bIsDrivingVelocityOfSimulatedTransformables;
+			for ( TUniquePtr<FViewportTransformable>& TransformablePtr : Transformables )
+			{
+				FViewportTransformable& Transformable = *TransformablePtr;
+
+				Transformable.LastTransform = Transformable.TargetTransform;
+
+				// Got an actor?
+				AActor* Actor = Cast<AActor>( Transformable.Object.Get() );
+				if ( Actor != nullptr )
+				{
+					const FTransform& ExistingTransform = Actor->GetTransform();
+
+					const bool bOnlyTranslationChanged =
+						ExistingTransform.GetRotation() == Transformable.TargetTransform.GetRotation() &&
+						ExistingTransform.GetScale3D() == Transformable.TargetTransform.GetScale3D();
+
+					Actor->SetActorTransform( Transformable.TargetTransform, bSweep );
+
+					Actor->InvalidateLightingCacheDetailed( bOnlyTranslationChanged );
+
+					const bool bFinished = false;
+					Actor->PostEditMove( bFinished );
+				}
+				else
+				{
+					// Some other object that we don't know how to drag
+				}
+			}
+		}
+	}
+
+	bIsFirstDragUpdate = false;
+}
+
+FVector UViewportWorldInteraction::ComputeConstrainedDragDeltaFromStart( const bool bIsFirstDragUpdate, const ETransformGizmoInteractionType InteractionType, const TOptional<FTransformGizmoHandlePlacement> OptionalHandlePlacement, const FVector& DragDeltaFromStart, const FVector& LaserPointerStart, const FVector& LaserPointerDirection, const bool bIsLaserPointerValid, const FTransform& GizmoStartTransform, FVector& GizmoSpaceFirstDragUpdateOffsetAlongAxis, FVector& GizmoSpaceDragDeltaFromStartOffset ) const
+{
+	FVector ConstrainedGizmoSpaceDeltaFromStart;
+
+	bool bConstrainDirectlyToLineOrPlane = false;
+	if ( OptionalHandlePlacement.IsSet() )
+	{
+		// Constrain to line/plane
+		const FTransformGizmoHandlePlacement& HandlePlacement = OptionalHandlePlacement.GetValue();
+
+		int32 CenterHandleCount, FacingAxisIndex, CenterAxisIndex;
+		HandlePlacement.GetCenterHandleCountAndFacingAxisIndex( /* Out */ CenterHandleCount, /* Out */ FacingAxisIndex, /* Out */ CenterAxisIndex );
+
+		// Our laser pointer ray won't be valid for inertial transform, since it's already moved on to other things.  But the existing velocity should already be on axis.
+		bConstrainDirectlyToLineOrPlane = ( CenterHandleCount == 2 ) && bIsLaserPointerValid;
+		if ( bConstrainDirectlyToLineOrPlane )
+		{
+			const FVector GizmoSpaceLaserPointerStart = GizmoStartTransform.InverseTransformPosition( LaserPointerStart );
+			const FVector GizmoSpaceLaserPointerDirection = GizmoStartTransform.InverseTransformVectorNoScale( LaserPointerDirection );
+
+			const float LaserPointerMaxLength = 1000; //Owner.GetLaserPointerMaxLength(); //@todo viewportinteraction: Needs proper way of getting distance
+
+			FVector GizmoSpaceConstraintAxis =
+				FacingAxisIndex == 0 ? FVector::ForwardVector : ( FacingAxisIndex == 1 ? FVector::RightVector : FVector::UpVector );
+			if ( HandlePlacement.Axes[ FacingAxisIndex ] == ETransformGizmoHandleDirection::Negative )
+			{
+				GizmoSpaceConstraintAxis *= -1.0f;
+			}
+
+			const bool bOnPlane = ( InteractionType == ETransformGizmoInteractionType::TranslateOnPlane );
+			if ( bOnPlane )
+			{
+				const FPlane GizmoSpaceConstrainToPlane( FVector::ZeroVector, GizmoSpaceConstraintAxis );
+
+				// 2D Plane
+				FVector GizmoSpacePointOnPlane = FVector::ZeroVector;
+				if ( !FMath::SegmentPlaneIntersection(
+					GizmoSpaceLaserPointerStart,
+					GizmoSpaceLaserPointerStart + LaserPointerMaxLength * GizmoSpaceLaserPointerDirection,
+					GizmoSpaceConstrainToPlane,
+					/* Out */ GizmoSpacePointOnPlane ) )
+				{
+					// No intersect.  Default to no delta.
+					GizmoSpacePointOnPlane = FVector::ZeroVector;
+				}
+
+				// Gizmo space is defined as the starting position of the interaction, so we simply take the closest position
+				// along the plane as our delta from the start in gizmo space
+				ConstrainedGizmoSpaceDeltaFromStart = GizmoSpacePointOnPlane;
+			}
+			else
+			{
+				// 1D Line
+				const float AxisSegmentLength = LaserPointerMaxLength * 2.0f;	// @todo vreditor: We're creating a line segment to represent an infinitely long axis, but really it just needs to be further than our laser pointer can reach
+
+				DVector GizmoSpaceClosestPointOnRayDouble, GizmoSpaceClosestPointOnAxisDouble;
+				SegmentDistToSegmentDouble(
+					GizmoSpaceLaserPointerStart, DVector( GizmoSpaceLaserPointerStart ) + DVector( GizmoSpaceLaserPointerDirection ) * LaserPointerMaxLength,	// @todo vreditor: Should store laser pointer length rather than hard code
+					DVector( GizmoSpaceConstraintAxis ) * -AxisSegmentLength, DVector( GizmoSpaceConstraintAxis ) * AxisSegmentLength,
+					/* Out */ GizmoSpaceClosestPointOnRayDouble,
+					/* Out */ GizmoSpaceClosestPointOnAxisDouble );
+
+				// Gizmo space is defined as the starting position of the interaction, so we simply take the closest position
+				// along the axis as our delta from the start in gizmo space
+				ConstrainedGizmoSpaceDeltaFromStart = GizmoSpaceClosestPointOnAxisDouble.ToFVector();
+			}
+
+			// Account for the handle position on the outside of the bounds.  This is a bit confusing but very important for dragging
+			// to feel right.  When constraining movement to either a plane or single axis, we always want the object to move exactly 
+			// to the location under the laser/cursor -- it's an absolute movement.  But if we did that on the first update frame, it 
+			// would teleport from underneath the gizmo handle to that location. 
+			{
+				if ( bIsFirstDragUpdate )
+				{
+					GizmoSpaceFirstDragUpdateOffsetAlongAxis = ConstrainedGizmoSpaceDeltaFromStart;
+				}
+				ConstrainedGizmoSpaceDeltaFromStart -= GizmoSpaceFirstDragUpdateOffsetAlongAxis;
+
+				// OK, it gets even more complicated.  When dragging and releasing for inertial movement, this code path
+				// will no longer be executed (because the hand/laser has moved on to doing other things, and our drag
+				// is driven by velocity only).  So we need to remember the LAST offset from the object's position to
+				// where we actually constrained it to, and continue to apply that delta during the inertial drag.
+				// That actually happens in the code block near the bottom of this function.
+				const FVector GizmoSpaceDragDeltaFromStart = GizmoStartTransform.InverseTransformVectorNoScale( DragDeltaFromStart );
+				GizmoSpaceDragDeltaFromStartOffset = ConstrainedGizmoSpaceDeltaFromStart - GizmoSpaceDragDeltaFromStart;
+			}
+		}
+	}
+
+	if ( !bConstrainDirectlyToLineOrPlane )
+	{
+		ConstrainedGizmoSpaceDeltaFromStart = GizmoStartTransform.InverseTransformVectorNoScale( DragDeltaFromStart );
+
+		// Apply axis constraint if we have one (and we haven't already constrained directly to a line)
+		if ( OptionalHandlePlacement.IsSet() )
+		{
+			// OK in this case, inertia is moving our selected objects while they remain constrained to
+			// either a plane or a single axis.  Every frame, we need to apply the original delta between
+			// where the laser was aiming and where the object was constrained to on the LAST frame before
+			// the user released the object and it moved inertially.  See the big comment up above for 
+			// more information.  Inertial movement of constrained objects is actually very complicated!
+			ConstrainedGizmoSpaceDeltaFromStart += GizmoSpaceDragDeltaFromStartOffset;
+
+			const bool bOnPlane = ( InteractionType == ETransformGizmoInteractionType::TranslateOnPlane );
+
+			const FTransformGizmoHandlePlacement& HandlePlacement = OptionalHandlePlacement.GetValue();
+			for ( int32 AxisIndex = 0; AxisIndex < 3; ++AxisIndex )
+			{
+				if ( bOnPlane )
+				{
+					if ( HandlePlacement.Axes[ AxisIndex ] == ETransformGizmoHandleDirection::Positive )
+					{
+						// Lock the opposing axes out (plane translation)
+						ConstrainedGizmoSpaceDeltaFromStart[ AxisIndex ] = 0.0f;
+					}
+				}
+				else
+				{
+					if ( HandlePlacement.Axes[ AxisIndex ] == ETransformGizmoHandleDirection::Center )
+					{
+						// Lock the centered axis out (line translation)
+						ConstrainedGizmoSpaceDeltaFromStart[ AxisIndex ] = 0.0f;
+					}
+				}
+			}
+		}
+	}
+
+	FVector ConstrainedWorldSpaceDeltaFromStart = GizmoStartTransform.TransformVectorNoScale( ConstrainedGizmoSpaceDeltaFromStart );
+	return ConstrainedWorldSpaceDeltaFromStart;
+}
+
+void UViewportWorldInteraction::StartDraggingActors( UViewportInteractor* Interactor, const FViewportActionKeyInput& Action, UActorComponent* ClickedComponent, const FVector& HitLocation, const bool bIsPlacingActors )
+{
+	if ( IsInteractableComponent( ClickedComponent ) )
+	{
+		FVector LaserPointerStart, LaserPointerEnd;
+		const bool bHaveLaserPointer = Interactor->GetLaserPointer( /* Out */ LaserPointerStart, /* Out */ LaserPointerEnd );
+		if ( bHaveLaserPointer )
+		{
+			AActor* Actor = ClickedComponent->GetOwner();
+
+			FViewportInteractorData& InteractorData = Interactor->GetInteractorData();
+
+			// Capture undo state
+			if ( bIsPlacingActors )
+			{
+				// When placing actors, a transaction should already be in progress
+				check( TrackingTransaction.IsActive() );
+			}
+			else
+			{
+				TrackingTransaction.TransCount++;
+				TrackingTransaction.Begin( LOCTEXT( "MovingActors", "Moving Selected Actors" ) );
+
+				// Suspend actor/component modification during each delta step to avoid recording unnecessary overhead into the transaction buffer
+				GEditor->DisableDeltaModification( true );
+			}
+
+			// Give the interacter a chance to do something when starting dragging
+			Interactor->OnStartDragging( ClickedComponent, HitLocation, bIsPlacingActors );
+
+			const bool bUsingGizmo =
+				Action.ActionType == ViewportWorldActionTypes::SelectAndMove_LightlyPressed &&		// Only use the gizmo when lightly pressed
+				( Actor == TransformGizmoActor ) &&
+				ClickedComponent != nullptr;
+
+			// Start dragging the objects right away!
+			InteractorData.DraggingMode = InteractorData.LastDraggingMode = bIsPlacingActors ? EViewportInteractionDraggingMode::ActorsAtLaserImpact : 
+				( bUsingGizmo ? EViewportInteractionDraggingMode::ActorsWithGizmo : EViewportInteractionDraggingMode::ActorsFreely );
+
+			// Starting a new drag, so make sure the other hand doesn't think it's assisting us
+			FViewportInteractorData& OtherInteractorData = Interactor->GetOtherInteractor()->GetInteractorData();
+			OtherInteractorData.bWasAssistingDrag = false;
+
+			InteractorData.bIsFirstDragUpdate = true;
+			InteractorData.bWasAssistingDrag = false;
+			InteractorData.DragRayLength = ( HitLocation - LaserPointerStart ).Size();
+			InteractorData.LastLaserPointerStart = LaserPointerStart;
+			InteractorData.LastDragToLocation = HitLocation;
+			InteractorData.LaserPointerImpactAtDragStart = HitLocation;
+			InteractorData.DragTranslationVelocity = FVector::ZeroVector;
+			InteractorData.DragRayLengthVelocity = 0.0f;
+			InteractorData.bIsDrivingVelocityOfSimulatedTransformables = false;
+
+			// Start dragging the transform gizmo.  Even if the user clicked on the actor itself, we'll use
+			// the gizmo to transform it.
+			if ( bUsingGizmo )
+			{
+				InteractorData.DraggingTransformGizmoComponent = Cast<USceneComponent>( ClickedComponent );
+			}
+			else
+			{
+				InteractorData.DraggingTransformGizmoComponent = nullptr;
+			}
+
+			if ( TransformGizmoActor != nullptr )
+			{
+				InteractorData.TransformGizmoInteractionType = TransformGizmoActor->GetInteractionType( InteractorData.DraggingTransformGizmoComponent.Get(), InteractorData.OptionalHandlePlacement );
+				InteractorData.GizmoStartTransform = TransformGizmoActor->GetTransform();
+				InteractorData.GizmoStartLocalBounds = GizmoLocalBounds;
+			}
+			else
+			{
+				InteractorData.TransformGizmoInteractionType = ETransformGizmoInteractionType::None;
+				InteractorData.GizmoStartTransform = FTransform::Identity;
+				InteractorData.GizmoStartLocalBounds = FBox( 0 );
+			}
+			InteractorData.GizmoSpaceFirstDragUpdateOffsetAlongAxis = FVector::ZeroVector;	// Will be determined on first update
+			InteractorData.GizmoSpaceDragDeltaFromStartOffset = FVector::ZeroVector;	// Set every frame while dragging
+
+			bDraggedSinceLastSelection = true;
+			LastDragGizmoStartTransform = InteractorData.GizmoStartTransform;
+
+			SetupTransformablesForSelectedActors();
+
+			// If we're placing actors, start interpolating to their actual location.  This helps smooth everything out when
+			// using the laser impact point as the target transform
+			if ( bIsPlacingActors )
+			{
+				bIsInterpolatingTransformablesFromSnapshotTransform = true;
+				const FTimespan CurrentTime = FTimespan::FromSeconds( FPlatformTime::Seconds() );
+				TransformablesInterpolationStartTime = CurrentTime;
+				TransformablesInterpolationDuration = VI::PlacementInterpolationDuration->GetFloat();
+			}
+
+			// Play a haptic effect when objects are picked up
+			Interactor->PlayHapticEffect( Interactor->GetDragHapticFeedbackStrength() );
+		}
+	}
+}
+void UViewportWorldInteraction::StopDragging( UViewportInteractor* Interactor )
+{
+	FViewportInteractorData& InteractorData = Interactor->GetInteractorData();
+	if ( InteractorData.DraggingMode != EViewportInteractionDraggingMode::Nothing )
+	{
+		// If the other hand started dragging after we started, allow that hand to "take over" the drag, so the user
+		// doesn't have to click again to continue their action.  Inertial effects of the hand that stopped dragging
+		// will still be in effect.
+		FViewportInteractorData& OtherInteractorData = Interactor->GetOtherInteractor()->GetInteractorData();
+		if ( OtherInteractorData.DraggingMode == EViewportInteractionDraggingMode::AssistingDrag )
+		{
+			// The other hand takes over whatever this hand was doing
+			OtherInteractorData.DraggingMode = OtherInteractorData.LastDraggingMode = InteractorData.DraggingMode;
+			OtherInteractorData.bIsDrivingVelocityOfSimulatedTransformables = InteractorData.bIsDrivingVelocityOfSimulatedTransformables;
+
+			// The other hand is no longer assisting, as it's now the primary interacting hand.
+			OtherInteractorData.bWasAssistingDrag = false;
+
+			// The hand that stopped dragging will be remembered as "assisting" the other hand, so that its
+			// inertia will still influence interactions
+			InteractorData.bWasAssistingDrag = true;
+		}
+		else
+		{
+			if ( InteractorData.DraggingMode == EViewportInteractionDraggingMode::Interactable )
+			{
+				if ( DraggedInteractable )
+				{
+					DraggedInteractable->OnDragRelease( Interactor );
+					UViewportDragOperationComponent* DragOperationComponent =  DraggedInteractable->GetDragOperationComponent();
+					if ( DragOperationComponent )
+					{
+						DragOperationComponent->ClearDragOperation();
+					}
+				}
+			}
+			else if ( InteractorData.DraggingMode == EViewportInteractionDraggingMode::ActorsWithGizmo ||
+				InteractorData.DraggingMode == EViewportInteractionDraggingMode::ActorsFreely ||
+				InteractorData.DraggingMode == EViewportInteractionDraggingMode::ActorsAtLaserImpact )
+			{
+				// No hand is dragging the objects anymore.
+				if ( InteractorData.DraggingMode == EViewportInteractionDraggingMode::ActorsWithGizmo &&
+					InteractorData.TransformGizmoInteractionType == ETransformGizmoInteractionType::RotateOnAngle )
+				{
+					StartDragAngleOnRotation.Reset();
+					StartDragHandleDirection.Reset();
+				}
+
+				// Finalize undo
+				{
+					// @todo vreditor undo: This doesn't actually encapsulate any inertial movement that happens after the drag is released! 
+					// We need to figure out whether that matters or not.  Also, look into PostEditMove( bFinished=true ) and how that relates to this.
+					--TrackingTransaction.TransCount;
+					TrackingTransaction.End();
+					GEditor->DisableDeltaModification( false );
+				}
+			}
+		}
+		
+		InteractorData.DraggingMode = EViewportInteractionDraggingMode::Nothing;
+		
+		FViewportActionKeyInput* OptionalAction = Interactor->GetActionWithName( ViewportWorldActionTypes::SelectAndMove );
+		if ( OptionalAction )
+		{
+			OptionalAction->bIsInputCaptured = false;
+		}
+	}
+
+	// NOTE: Even though we're done dragging, we keep our list of transformables so that inertial drag can still push them around!
+}
+
+void UViewportWorldInteraction::SetupTransformablesForSelectedActors()
+{
+	//@todo ViewportInteraction: Make sure this is not only for actors
+	bIsInterpolatingTransformablesFromSnapshotTransform = false;
+	TransformablesInterpolationStartTime = FTimespan::Zero();
+	TransformablesInterpolationDuration = 1.0f;
+
+	Transformables.Reset();
+
+	USelection* ActorSelectionSet = GEditor->GetSelectedActors();
+
+	static TArray<UObject*> SelectedActorObjects;
+	SelectedActorObjects.Reset();
+	ActorSelectionSet->GetSelectedObjects( AActor::StaticClass(), /* Out */ SelectedActorObjects );
+
+	for ( UObject* SelectedActorObject : SelectedActorObjects )
+	{
+		AActor* SelectedActor = CastChecked<AActor>( SelectedActorObject );
+
+		Transformables.Add( TUniquePtr< FViewportTransformable>( new FViewportTransformable() ) );
+		FViewportTransformable& Transformable = *Transformables[ Transformables.Num() - 1 ];
+
+		Transformable.Object = SelectedActor;
+		Transformable.StartTransform = Transformable.LastTransform = Transformable.TargetTransform = Transformable.UnsnappedTargetTransform = Transformable.InterpolationSnapshotTransform = SelectedActor->GetTransform();
+	}
+}
+
+bool UViewportWorldInteraction::FindPlacementPointUnderLaser( UViewportInteractor* Interactor, FVector& OutHitLocation )
+{
+	// Check to see if the laser pointer is over something we can drop on
+	bool bHitSomething = false;
+	FVector HitLocation = FVector::ZeroVector;
+
+	static TArray<AActor*> IgnoredActors;
+	{
+		IgnoredActors.Reset();
+
+		// Don't trace against stuff that we're dragging!
+		for( TUniquePtr<FViewportTransformable>& TransformablePtr : Transformables )
+		{
+			FViewportTransformable& Transformable = *TransformablePtr;
+
+			// Got an actor?
+			AActor* Actor = Cast<AActor>( Transformable.Object.Get() );
+			if( Actor != nullptr )
+			{
+				IgnoredActors.Add( Actor );
+			}
+		} 
+	}
+
+
+	const bool bIgnoreGizmos = true;		// Never place on top of gizmos, just ignore them
+	const bool bEvenIfUIIsInFront = true;	// Don't let the UI block placement
+	FHitResult HitResult = Interactor->GetHitResultFromLaserPointer( &IgnoredActors, bIgnoreGizmos, nullptr, bEvenIfUIIsInFront );
+	if( HitResult.Actor.IsValid() )
+	{
+		bHitSomething = true;
+		HitLocation = HitResult.ImpactPoint;
+	}
+
+	if( bHitSomething )
+	{
+		FViewportInteractorData& InteractorData = Interactor->GetInteractorData();
+
+		// Pull back from the impact point
+		{
+			const FVector GizmoSpaceImpactNormal = InteractorData.GizmoStartTransform.InverseTransformVectorNoScale( HitResult.ImpactNormal );
+
+			FVector ExtremePointOnBox;
+			const FVector ExtremeDirection = -GizmoSpaceImpactNormal;
+			const FBox& Box = GizmoLocalBounds;
+			{
+				const FVector ExtremePoint(
+					ExtremeDirection.X >= 0.0f ? Box.Max.X : Box.Min.X,
+					ExtremeDirection.Y >= 0.0f ? Box.Max.Y : Box.Min.Y,
+					ExtremeDirection.Z >= 0.0f ? Box.Max.Z : Box.Min.Z );
+				const float ProjectionDistance = FVector::DotProduct( ExtremePoint, ExtremeDirection );
+				ExtremePointOnBox = ExtremeDirection * ProjectionDistance;
+			}
+
+			const FVector WorldSpaceExtremePointOnBox = InteractorData.GizmoStartTransform.TransformVectorNoScale( ExtremePointOnBox );
+			HitLocation -= WorldSpaceExtremePointOnBox;
+		}
+
+		OutHitLocation = HitLocation;
+	}
+
+	return bHitSomething;
+}
+
+bool UViewportWorldInteraction::IsTransformingActor( AActor* Actor ) const
+{
+	bool bFoundActor = false;
+	for ( const TUniquePtr<FViewportTransformable>& TransformablePtr : Transformables )
+	{
+		const FViewportTransformable& Transformable = *TransformablePtr;
+
+		if ( Transformable.Object.Get() == Actor )
+		{
+			bFoundActor = true;
+			break;
+		}
+	}
+
+	return bFoundActor;
+}
+
+bool UViewportWorldInteraction::IsSmoothSnappingEnabled() const
+{
+	const float SmoothSnapSpeed = VI::SmoothSnapSpeed->GetFloat();
+	const bool bSmoothSnappingEnabled =
+		!GEditor->bIsSimulatingInEditor &&
+		( FSnappingUtils::IsSnapToGridEnabled() || FSnappingUtils::IsRotationSnapEnabled() || FSnappingUtils::IsScaleSnapEnabled() ) &&
+		VI::SmoothSnap->GetInt() != 0 &&
+		!FMath::IsNearlyZero( SmoothSnapSpeed );
+
+	return bSmoothSnappingEnabled;
+}
+
+void UViewportWorldInteraction::PollInputIfNeeded()
+{
+	if ( LastFrameNumberInputWasPolled != GFrameNumber )
+	{
+		for ( UViewportInteractor* Interactor : Interactors )
+		{
+			Interactor->PollInput();
+		}
+
+		LastFrameNumberInputWasPolled = GFrameNumber;
+	}
+}
+
+void UViewportWorldInteraction::OnActorSelectionChanged( UObject* ChangedObject )
+{
+	if ( EditorViewportClient != nullptr )
+	{
+		const bool bNewObjectsSelected = true;
+		const bool bAllHandlesVisible = true;
+		UActorComponent* SingleVisibleHandle = nullptr;
+		TArray<UActorComponent*> HoveringOverHandles;
+		RefreshOnSelectionChanged( bNewObjectsSelected, bAllHandlesVisible, SingleVisibleHandle, HoveringOverHandles );
+
+		// Clear our transformables.  They'll be doing things relative to the gizmo, which will now be in a different place.
+		// @todo vreditor: We could solve this by moving GizmoLocalBounds into the a Transformables base structure, and making sure
+		// that we never talk to TransformGizmoActor (or selected actors) when updating in tick (currently safe!)
+		Transformables.Reset();
+
+		bDraggedSinceLastSelection = false;
+		LastDragGizmoStartTransform = FTransform::Identity;
+	}
+}
+
+
+void UViewportWorldInteraction::RefreshTransformGizmo( const bool bNewObjectsSelected, bool bAllHandlesVisible, class UActorComponent* SingleVisibleHandle, const TArray< UActorComponent* >& HoveringOverHandles )
+{
+	if ( GEditor->GetSelectedActorCount() > 0 )
+	{
+		SpawnTransformGizmoIfNeeded();
+
+		const ECoordSystem CurrentCoordSystem = GetTransformGizmoCoordinateSpace();
+		const bool bIsWorldSpaceGizmo = ( CurrentCoordSystem == COORD_World );
+
+		USelection* ActorSelectionSet = GEditor->GetSelectedActors();
+
+		static TArray<UObject*> SelectedActorObjects;
+		SelectedActorObjects.Reset();
+		ActorSelectionSet->GetSelectedObjects( AActor::StaticClass(), /* Out */ SelectedActorObjects );
+
+		AActor* LastSelectedActor = CastChecked<AActor>( SelectedActorObjects.Last() );
+
+		// Use the location and orientation of the last selected actor, to be consistent with the regular editor's gizmo
+		FTransform GizmoToWorld = LastSelectedActor->GetTransform();
+		GizmoToWorld.RemoveScaling();	// We don't need the pivot itself to be scaled
+
+		if ( bIsWorldSpaceGizmo )
+		{
+			GizmoToWorld.SetRotation( FQuat::Identity );
+		}
+
+		// Create a gizmo-local bounds around all of the selected actors
+		FBox GizmoSpaceSelectedActorsBounds;
+		GizmoSpaceSelectedActorsBounds.Init();
+		{
+			const FTransform WorldToGizmo = GizmoToWorld.Inverse();
+			for ( UObject* SelectedActorObject : SelectedActorObjects )
+			{
+				AActor* SelectedActor = CastChecked<AActor>( SelectedActorObject );
+
+				// Figure out the local space bounds of the entire actor
+				const bool bIncludeNonCollidingComponents = false;	// @todo vreditor: Disabled this because it causes lights to have huge bounds
+				const FBox ActorSpaceBounds = SelectedActor->CalculateComponentsBoundingBoxInLocalSpace( bIncludeNonCollidingComponents );
+
+				// Transform the bounds into the gizmo's space
+				const FTransform ActorToWorld = SelectedActor->GetTransform();
+				const FTransform ActorToGizmo = ActorToWorld * WorldToGizmo;
+				const FBox GizmoSpaceBounds = ActorSpaceBounds.TransformBy( ActorToGizmo );
+
+				// Get that bounding box into the local space of the gizmo
+				GizmoSpaceSelectedActorsBounds += GizmoSpaceBounds;
+			}
+		}
+
+		if ( VI::ForceGizmoPivotToCenterOfSelectedActorsBounds->GetInt() > 0 )
+		{
+			GizmoToWorld.SetLocation( GizmoToWorld.TransformPosition( GizmoSpaceSelectedActorsBounds.GetCenter() ) );
+			GizmoSpaceSelectedActorsBounds = GizmoSpaceSelectedActorsBounds.ShiftBy( -GizmoSpaceSelectedActorsBounds.GetCenter() );
+		}
+
+		if ( bNewObjectsSelected )
+		{
+			TransformGizmoActor->OnNewObjectsSelected();
+		}
+
+		const EGizmoHandleTypes GizmoType = GetCurrentGizmoType();
+		const ECoordSystem GizmoCoordinateSpace = GetTransformGizmoCoordinateSpace();
+
+		GizmoLocalBounds = GizmoSpaceSelectedActorsBounds;
+
+		TransformGizmoActor->UpdateGizmo( GizmoType, GizmoCoordinateSpace, GizmoToWorld, GizmoSpaceSelectedActorsBounds, GetHeadTransform().GetLocation(), bAllHandlesVisible, SingleVisibleHandle,
+			HoveringOverHandles, VI::GizmoHandleHoverScale->GetFloat(), VI::GizmoHandleHoverAnimationDuration->GetFloat() );
+
+		// Only draw if snapping is turned on
+		AActor* SnapGridActor = GetSnapGridActor();
+		if ( FSnappingUtils::IsSnapToGridEnabled() )
+		{
+			SnapGridActor->GetRootComponent()->SetVisibility( true );
+
+			const float GizmoAnimationAlpha = TransformGizmoActor->GetAnimationAlpha();
+
+			// Position the grid snap actor
+			const FVector GizmoLocalBoundsBottom( 0.0f, 0.0f, GizmoLocalBounds.Min.Z );
+			const float SnapGridZOffset = 0.1f;	// @todo vreditor tweak: Offset the grid position a little bit to avoid z-fighting with objects resting directly on a floor
+			const FVector SnapGridLocation = GizmoToWorld.TransformPosition( GizmoLocalBoundsBottom ) + FVector( 0.0f, 0.0f, SnapGridZOffset );
+			SnapGridActor->SetActorLocationAndRotation( SnapGridLocation, GizmoToWorld.GetRotation() );
+
+			// Figure out how big to make the snap grid.  We'll use a multiplier of the object's bounding box size (ignoring local 
+			// height, because we currently only draw the grid at the bottom of the object.)
+			FBox GizmoLocalBoundsFlattened = GizmoLocalBounds;
+			GizmoLocalBoundsFlattened.Min.Z = GizmoLocalBoundsFlattened.Max.Z = 0.0f;
+			const FBox GizmoWorldBoundsFlattened = GizmoLocalBoundsFlattened.TransformBy( GizmoToWorld );
+			const float SnapGridSize = GizmoWorldBoundsFlattened.GetSize().GetAbsMax() * VI::SnapGridSize->GetFloat();
+
+			// The mesh is 100x100, so we'll scale appropriately
+			const float SnapGridScale = SnapGridSize / 100.0f;
+			SnapGridActor->SetActorScale3D( FVector( SnapGridScale ) );
+
+			EViewportInteractionDraggingMode DraggingMode = EViewportInteractionDraggingMode::Nothing;
+			FVector GizmoStartLocationWhileDragging = FVector::ZeroVector;
+			for ( UViewportInteractor* Interactor : Interactors )
+			{
+				const FViewportInteractorData& InteractorData = Interactor->GetInteractorData();
+				if ( InteractorData.DraggingMode == EViewportInteractionDraggingMode::ActorsFreely ||
+					InteractorData.DraggingMode == EViewportInteractionDraggingMode::ActorsWithGizmo ||
+					InteractorData.DraggingMode == EViewportInteractionDraggingMode::ActorsAtLaserImpact )
+				{
+					DraggingMode = InteractorData.DraggingMode;
+					GizmoStartLocationWhileDragging = InteractorData.GizmoStartTransform.GetLocation();
+					break;
+				}
+				else if ( bDraggedSinceLastSelection &&
+					( InteractorData.LastDraggingMode == EViewportInteractionDraggingMode::ActorsFreely ||
+					InteractorData.LastDraggingMode == EViewportInteractionDraggingMode::ActorsWithGizmo ||
+					InteractorData.LastDraggingMode == EViewportInteractionDraggingMode::ActorsAtLaserImpact ) )
+				{
+					DraggingMode = InteractorData.LastDraggingMode;
+					GizmoStartLocationWhileDragging = LastDragGizmoStartTransform.GetLocation();
+				}
+			}
+
+			FVector SnapGridCenter = SnapGridLocation;
+			if ( DraggingMode != EViewportInteractionDraggingMode::Nothing )
+			{
+				SnapGridCenter = GizmoStartLocationWhileDragging;
+			}
+
+			// Fade the grid in with the transform gizmo
+			const float GridAlpha = GizmoAnimationAlpha;
+
+			// Animate the grid a little bit
+			const FLinearColor GridColor = FLinearColor::LerpUsingHSV(
+				FLinearColor::White,
+				FLinearColor::Yellow,
+				FMath::MakePulsatingValue( GetTimeSinceEntered().GetTotalSeconds(), 0.5f ) ).CopyWithNewOpacity( GridAlpha );
+
+			const float GridInterval = GEditor->GetGridSize();
+
+			// If the grid size is very small, shrink the size of our lines so that they don't overlap
+			float LineWidth = VI::SnapGridLineWidth->GetFloat();
+			while ( GridInterval < LineWidth * 3.0f )	// @todo vreditor tweak
+			{
+				LineWidth *= 0.5f;
+			}
+
+			{
+				// Update snap grid material state
+				UMaterialInstanceDynamic* SnapGridMID = GetSnapGridMID();
+
+				static FName GridColorParameterName( "GridColor" );
+				SnapGridMID->SetVectorParameterValue( GridColorParameterName, GridColor );
+
+				static FName GridCenterParameterName( "GridCenter" );
+				SnapGridMID->SetVectorParameterValue( GridCenterParameterName, SnapGridCenter );
+
+				static FName GridIntervalParameterName( "GridInterval" );
+				SnapGridMID->SetScalarParameterValue( GridIntervalParameterName, GridInterval );
+
+				static FName GridRadiusParameterName( "GridRadius" );
+				SnapGridMID->SetScalarParameterValue( GridRadiusParameterName, SnapGridSize * 0.5f );
+
+				static FName LineWidthParameterName( "LineWidth" );
+				SnapGridMID->SetScalarParameterValue( LineWidthParameterName, LineWidth );
+
+				FMatrix GridRotationMatrix = GizmoToWorld.ToMatrixNoScale().Inverse();
+				GridRotationMatrix.SetOrigin( FVector::ZeroVector );
+
+				static FName GridRotationMatrixXParameterName( "GridRotationMatrixX" );
+				static FName GridRotationMatrixYParameterName( "GridRotationMatrixY" );
+				static FName GridRotationMatrixZParameterName( "GridRotationMatrixZ" );
+				SnapGridMID->SetVectorParameterValue( GridRotationMatrixXParameterName, GridRotationMatrix.GetScaledAxis( EAxis::X ) );
+				SnapGridMID->SetVectorParameterValue( GridRotationMatrixYParameterName, GridRotationMatrix.GetScaledAxis( EAxis::Y ) );
+				SnapGridMID->SetVectorParameterValue( GridRotationMatrixZParameterName, GridRotationMatrix.GetScaledAxis( EAxis::Z ) );
+			}
+		}
+		else
+		{
+			// Grid snap not enabled
+			SnapGridActor->GetRootComponent()->SetVisibility( false );
+		}
+	}
+	else
+	{
+		// Nothing selected, so kill our gizmo actor
+		if ( TransformGizmoActor != nullptr )
+		{
+			DestroyTransientActor( TransformGizmoActor );
+			TransformGizmoActor = nullptr;
+		}
+		GizmoLocalBounds = FBox( 0 );
+
+		// Hide the snap actor
+		AActor* SnapGridActor = GetSnapGridActor();
+		SnapGridActor->GetRootComponent()->SetVisibility( false );
+	}
+}
+
+void UViewportWorldInteraction::SpawnTransformGizmoIfNeeded()
+{
+	// Check if there no gizmo or if the wrong gizmo is being used
+	bool bSpawnNewGizmo = false;
+	if ( TransformGizmoActor == nullptr || TransformGizmoActor->GetClass() != TransformGizmoClass )
+	{
+		bSpawnNewGizmo = true;
+	}
+
+	if ( bSpawnNewGizmo )
+	{
+		// Destroy the previous gizmo
+		if ( TransformGizmoActor != nullptr )
+		{
+			TransformGizmoActor->Destroy();
+		}
+
+		// Create the correct gizmo
+		const bool bWithSceneComponent = false;	// We already have our own scene component
+		TransformGizmoActor = CastChecked<ABaseTransformGizmo>( SpawnTransientSceneActor( TransformGizmoClass, TEXT( "PivotTransformGizmo" ), bWithSceneComponent ) );
+
+		check( TransformGizmoActor != nullptr );
+		TransformGizmoActor->SetOwnerWorldInteraction( this );
+
+		if ( VI::ShowTransformGizmo->GetInt() == 0 )
+		{
+			TransformGizmoActor->SetIsTemporarilyHiddenInEditor( true );
+		}
+	}
+}
+
+void UViewportWorldInteraction::ApplyVelocityDamping( FVector& Velocity, const bool bVelocitySensitive )
+{
+	const float InertialMovementZeroEpsilon = 0.01f;	// @todo vreditor tweak
+	if ( !Velocity.IsNearlyZero( InertialMovementZeroEpsilon ) )
+	{
+		// Apply damping
+		if ( bVelocitySensitive )
+		{
+			const float DampenMultiplierAtLowSpeeds = 0.94f;	// @todo vreditor tweak
+			const float DampenMultiplierAtHighSpeeds = 0.99f;	// @todo vreditor tweak
+			const float SpeedForMinimalDamping = 2.5f * GetWorldScaleFactor();	// cm/frame	// @todo vreditor tweak
+			const float SpeedBasedDampeningScalar = FMath::Clamp( Velocity.Size(), 0.0f, SpeedForMinimalDamping ) / SpeedForMinimalDamping;	// @todo vreditor: Probably needs a curve applied to this to compensate for our framerate insensitivity
+			Velocity = Velocity * FMath::Lerp( DampenMultiplierAtLowSpeeds, DampenMultiplierAtHighSpeeds, SpeedBasedDampeningScalar );	// @todo vreditor: Frame rate sensitive damping.  Make use of delta time!
+		}
+		else
+		{
+			Velocity = Velocity * 0.95f;
+		}
+	}
+
+	if ( Velocity.IsNearlyZero( InertialMovementZeroEpsilon ) )
+	{
+		Velocity = FVector::ZeroVector;
+	}
+}
+
+void UViewportWorldInteraction::CycleTransformGizmoCoordinateSpace()
+{
+	const bool bGetRawValue = true;
+	const ECoordSystem CurrentCoordSystem = GLevelEditorModeTools().GetCoordSystem( bGetRawValue );
+	SetTransformGizmoCoordinateSpace( CurrentCoordSystem == COORD_World ? COORD_Local : COORD_World );
+}
+
+void UViewportWorldInteraction::SetTransformGizmoCoordinateSpace( const ECoordSystem NewCoordSystem )
+{
+	GLevelEditorModeTools().SetCoordSystem( NewCoordSystem );
+}
+
+ECoordSystem UViewportWorldInteraction::GetTransformGizmoCoordinateSpace() const
+{
+	const bool bGetRawValue = true;
+	const ECoordSystem CurrentCoordSystem = GLevelEditorModeTools().GetCoordSystem( bGetRawValue );
+	return CurrentCoordSystem;
+}
+
+float UViewportWorldInteraction::GetMaxScale()
+{
+	return VI::ScaleMax->GetFloat();
+}
+
+float UViewportWorldInteraction::GetMinScale()
+{
+	return VI::ScaleMin->GetFloat();
+}
+
+void UViewportWorldInteraction::SetWorldToMetersScale( const float NewWorldToMetersScale )
+{
+	// @todo vreditor: This is bad because we're clobbering the world settings which will be saved with the map.  Instead we need to 
+	// be able to apply an override before the scene view gets it
+
+	ENGINE_API extern float GNewWorldToMetersScale;
+	GNewWorldToMetersScale = NewWorldToMetersScale;
+}
+
+UViewportInteractor* UViewportWorldInteraction::GetOtherInteractorIntertiaContribute( UViewportInteractor* Interactor )
+{
+	// Check to see if the other hand has any inertia to contribute
+	UViewportInteractor* OtherInteractorThatWasAssistingDrag = nullptr;
+	{
+		UViewportInteractor* OtherInteractor = Interactor->GetOtherInteractor();
+		FViewportInteractorData& OtherHandInteractorData = OtherInteractor->GetInteractorData();
+
+		// If the other hand isn't doing anything, but the last thing it was doing was assisting a drag, then allow it
+		// to contribute inertia!
+		if ( OtherHandInteractorData.DraggingMode == EViewportInteractionDraggingMode::Nothing && OtherHandInteractorData.bWasAssistingDrag )
+		{
+			if ( !OtherHandInteractorData.DragTranslationVelocity.IsNearlyZero( KINDA_SMALL_NUMBER ) )
+			{
+				OtherInteractorThatWasAssistingDrag = OtherInteractor;
+			}
+		}
+	}
+
+	return OtherInteractorThatWasAssistingDrag;
+}
+
+AActor* UViewportWorldInteraction::SpawnTransientSceneActor( TSubclassOf<AActor> ActorClass, const FString& ActorName, const bool bWithSceneComponent ) const
+{
+	UWorld* World = GetViewportWorld();
+	const bool bWasWorldPackageDirty = World->GetOutermost()->IsDirty();
+
+	// @todo vreditor: Needs respawn if world changes (map load, etc.)  Will that always restart the editor mode anyway?
+	FActorSpawnParameters ActorSpawnParameters;
+	ActorSpawnParameters.Name = MakeUniqueObjectName( World, ActorClass, *ActorName );	// @todo vreditor: Without this, SpawnActor() can return us an existing PendingKill actor of the same name!  WTF?
+	ActorSpawnParameters.SpawnCollisionHandlingOverride = ESpawnActorCollisionHandlingMethod::AlwaysSpawn;
+	ActorSpawnParameters.ObjectFlags = EObjectFlags::RF_Transient;
+
+	check( ActorClass != nullptr );
+	AActor* NewActor = World->SpawnActor< AActor >( ActorClass, ActorSpawnParameters );
+	NewActor->SetActorLabel( ActorName );
+
+	if ( bWithSceneComponent )
+	{
+		// Give the new actor a root scene component, so we can attach multiple sibling components to it
+		USceneComponent* SceneComponent = NewObject<USceneComponent>( NewActor );
+		NewActor->AddOwnedComponent( SceneComponent );
+		NewActor->SetRootComponent( SceneComponent );
+		SceneComponent->RegisterComponent();
+	}
+
+	// Don't dirty the level file after spawning a transient actor
+	if ( !bWasWorldPackageDirty )
+	{
+		GetViewportWorld()->GetOutermost()->SetDirtyFlag( false );
+	}
+
+	return NewActor;
+}
+
+
+void UViewportWorldInteraction::DestroyTransientActor( AActor* Actor ) const
+{
+	if ( Actor != nullptr )
+	{
+		const bool bWasWorldPackageDirty = GetViewportWorld()->GetOutermost()->IsDirty();
+
+		const bool bNetForce = false;
+		const bool bShouldModifyLevel = false;	// Don't modify level for transient actor destruction
+		GetViewportWorld()->DestroyActor( Actor, bNetForce, bShouldModifyLevel );
+		Actor = nullptr;
+
+		// Don't dirty the level file after destroying a transient actor
+		if ( !bWasWorldPackageDirty )
+		{
+			GetViewportWorld()->GetOutermost()->SetDirtyFlag( false );
+		}
+	}
+}
+
+EGizmoHandleTypes UViewportWorldInteraction::GetCurrentGizmoType() const
+{
+	return CurrentGizmoType;
+}
+
+void UViewportWorldInteraction::SetGizmoHandleType( const EGizmoHandleTypes InGizmoHandleType )
+{
+	CurrentGizmoType = InGizmoHandleType;
+}
+
+void UViewportWorldInteraction::SetTransformGizmoClass( const TSubclassOf<ABaseTransformGizmo>& NewTransformGizmoClass )
+{
+	TransformGizmoClass = NewTransformGizmoClass;
+}
+
+void UViewportWorldInteraction::SetDraggedInteractable( IViewportInteractableInterface* InDraggedInteractable )
+{
+	DraggedInteractable = InDraggedInteractable;
+	if ( DraggedInteractable && DraggedInteractable->GetDragOperationComponent() )
+	{
+		DraggedInteractable->GetDragOperationComponent()->StartDragOperation();
+	}
+}
+
+bool UViewportWorldInteraction::IsOtherInteractorHoveringOverComponent( UViewportInteractor* Interactor, UActorComponent* Component ) const
+{
+	bool bResult = false;
+
+	if ( Interactor && Component )
+	{
+		for ( UViewportInteractor* CurrentInteractor : Interactors )
+		{
+			if ( CurrentInteractor != Interactor && CurrentInteractor->GetHoverComponent() == Component )
+			{
+				bResult = true;
+				break;
+			}
+		}
+	}
+
+	return bResult;
+}
+
+void UViewportWorldInteraction::SpawnGridMeshActor()
+{
+	// Snap grid mesh
+	if( SnapGridActor == nullptr )
+	{
+		const bool bWithSceneComponent = false;
+		SnapGridActor = SpawnTransientSceneActor<AActor>( TEXT( "SnapGrid" ), bWithSceneComponent );
+
+		SnapGridMeshComponent = NewObject<UStaticMeshComponent>( SnapGridActor );
+		SnapGridActor->AddOwnedComponent( SnapGridMeshComponent );
+		SnapGridActor->SetRootComponent( SnapGridMeshComponent );
+		SnapGridMeshComponent->RegisterComponent();
+
+		UStaticMesh* GridMesh = LoadObject<UStaticMesh>( nullptr, TEXT( "/Engine/VREditor/SnapGrid/SnapGridPlaneMesh" ) );
+		check( GridMesh != nullptr );
+		SnapGridMeshComponent->SetStaticMesh( GridMesh );
+		SnapGridMeshComponent->SetMobility( EComponentMobility::Movable );
+		SnapGridMeshComponent->SetCollisionEnabled( ECollisionEnabled::NoCollision );
+
+		UMaterialInterface* GridMaterial = LoadObject<UMaterialInterface>( nullptr, TEXT( "/Engine/VREditor/SnapGrid/SnapGridMaterial" ) );
+		check( GridMaterial != nullptr );
+
+		SnapGridMID = UMaterialInstanceDynamic::Create( GridMaterial, GetTransientPackage() );
+		check( SnapGridMID != nullptr );
+		SnapGridMeshComponent->SetMaterial( 0, SnapGridMID );
+
+		// The grid starts off hidden
+		SnapGridMeshComponent->SetVisibility( false );
+	}
+}
+
+FLinearColor UViewportWorldInteraction::GetColor( const EColors Color ) const
+{
+	return Colors[ (int32)Color ];
+}
+
+#undef LOCTEXT_NAMESPACE
\ No newline at end of file
