VR Editor:

diff --git a/Engine/Source/Editor/ViewportInteraction/ViewportWorldInteractionInterface.cpp b/Engine/Source/Editor/ViewportInteraction/ViewportWorldInteractionInterface.cpp
new file mode 100644
--- /dev/null
+++ b/Engine/Source/Editor/ViewportInteraction/ViewportWorldInteractionInterface.cpp
@@ -0,0 +1,10 @@
+// Copyright 1998-2015 Epic Games, Inc. All Rights Reserved.
+
+#include "ViewportInteractionModule.h"
+#include "ViewportWorldInteractionInterface.h"
+
+UViewportWorldInteractionInterface::UViewportWorldInteractionInterface( const class FObjectInitializer& ObjectInitializer )
+	: Super( ObjectInitializer )
+{
+
+}
