VR Editor: - Fix mac warning: declaration shadows a field for SnapGridActor
- Fix SetHasUIOnForearm setter for bHasUIOnForearm

diff --git a/Engine/Source/Editor/VREditor/VREditorInteractor.cpp b/Engine/Source/Editor/VREditor/VREditorInteractor.cpp
--- a/Engine/Source/Editor/VREditor/VREditorInteractor.cpp
+++ b/Engine/Source/Editor/VREditor/VREditorInteractor.cpp
@@ -553,7 +553,7 @@ bool UVREditorInteractor::HasUIInFront() const
 
 void UVREditorInteractor::SetHasUIOnForearm( const bool bInHasUIOnForearm )
 {
-	bHasUIOnForearm = bHasUIOnForearm;
+	bHasUIOnForearm = bInHasUIOnForearm;
 }
 
 bool UVREditorInteractor::HasUIOnForearm() const
diff --git a/Engine/Source/Editor/ViewportInteraction/ViewportWorldInteraction.cpp b/Engine/Source/Editor/ViewportInteraction/ViewportWorldInteraction.cpp
--- a/Engine/Source/Editor/ViewportInteraction/ViewportWorldInteraction.cpp
+++ b/Engine/Source/Editor/ViewportInteraction/ViewportWorldInteraction.cpp
@@ -16,6 +16,7 @@
 // For actor placement
 #include "ObjectTools.h"
 #include "AssetSelection.h"
+#include "IHeadMountedDisplay.h"
 
 #define LOCTEXT_NAMESPACE "ViewportWorldInteraction"
 
@@ -209,13 +210,10 @@ static void SegmentDistToSegmentDouble( DVector A1, DVector B1, DVector A2, DVec
 
 UViewportWorldInteraction::UViewportWorldInteraction( const FObjectInitializer& Initializer ) :
 	Super( Initializer ),
-	bDraggedSinceLastSelection( false ),
-	LastDragGizmoStartTransform( FTransform::Identity ),
-	TrackingTransaction( FTrackingTransaction() ),
 	AppTimeEntered( FTimespan::Zero() ),
-	EditorViewportClient( nullptr ),
 	LastFrameNumberInputWasPolled( 0 ),
 	MotionControllerID( 0 ),	// @todo ViewportInteraction: We only support a single controller, and we assume the first controller are the motion controls
+	EditorViewportClient( nullptr ),
 	LastWorldToMetersScale( 100.0f ),
 	bIsInterpolatingTransformablesFromSnapshotTransform( false ),
 	TransformablesInterpolationStartTime( FTimespan::Zero() ),
@@ -223,6 +221,8 @@ UViewportWorldInteraction::UViewportWorldInteraction( const FObjectInitializer&
 	TransformGizmoActor( nullptr ),
 	TransformGizmoClass( APivotTransformGizmo::StaticClass() ),
 	GizmoLocalBounds( FBox( 0 ) ),
+	bDraggedSinceLastSelection( false ),
+	LastDragGizmoStartTransform( FTransform::Identity ),
 	StartDragAngleOnRotation( ),
 	StartDragHandleDirection( ),
 	CurrentGizmoType( EGizmoHandleTypes::All ),
@@ -2451,7 +2451,7 @@ void UViewportWorldInteraction::RefreshTransformGizmo( const bool bNewObjectsSel
 			HoveringOverHandles, VI::GizmoHandleHoverScale->GetFloat(), VI::GizmoHandleHoverAnimationDuration->GetFloat() );
 
 		// Only draw if snapping is turned on
-		AActor* SnapGridActor = GetSnapGridActor();
+		SpawnGridMeshActor();
 		if ( FSnappingUtils::IsSnapToGridEnabled() )
 		{
 			SnapGridActor->GetRootComponent()->SetVisibility( true );
@@ -2569,8 +2569,7 @@ void UViewportWorldInteraction::RefreshTransformGizmo( const bool bNewObjectsSel
 		GizmoLocalBounds = FBox( 0 );
 
 		// Hide the snap actor
-		AActor* SnapGridActor = GetSnapGridActor();
-		SnapGridActor->GetRootComponent()->SetVisibility( false );
+		GetSnapGridActor()->GetRootComponent()->SetVisibility( false );
 	}
 }
 
