VR Editor: Miscellaneous fixes (Merged from CL 3024520)
- Fixed input capture not being revoked properly after releasing controller triggers
- Fixed some shadowed variable warnings
- Renamed some function parameter variables

diff --git a/Engine/Source/Editor/MeshPaint/Private/MeshPaintEdMode.cpp b/Engine/Source/Editor/MeshPaint/Private/MeshPaintEdMode.cpp
--- a/Engine/Source/Editor/MeshPaint/Private/MeshPaintEdMode.cpp
+++ b/Engine/Source/Editor/MeshPaint/Private/MeshPaintEdMode.cpp
@@ -5060,7 +5060,7 @@ TWeakObjectPtr<AActor> FEdModeMeshPaint::GetEditingActor() const
 
 
 void FEdModeMeshPaint::OnVRAction( class FEditorViewportClient& ViewportClient, UViewportInteractor* Interactor, 
-	const FViewportActionKeyInput& Action, const EInputEvent Event, bool& bOutIsInCaptured, bool& bWasHandled )
+	const FViewportActionKeyInput& Action, const EInputEvent Event, bool& bOutIsInputCaptured, bool& bWasHandled )
 {
 	IVREditorMode* VREditorMode = static_cast<IVREditorMode*>( GetModeManager()->GetActiveMode( IVREditorModule::Get().GetVREditorModeID() ) );
 	UVREditorInteractor* VRInteractor = Cast<UVREditorInteractor>( Interactor );
@@ -5074,7 +5074,7 @@ void FEdModeMeshPaint::OnVRAction( class FEditorViewportClient& ViewportClient,
 			!VRInteractor->IsHoveringOverUI() )
 		{
 			bWasHandled = true;
-			bOutIsInCaptured = true;
+			bOutIsInputCaptured = true;
 		}
 
 		if( Action.ActionType == ViewportWorldActionTypes::SelectAndMove_LightlyPressed )
@@ -5086,7 +5086,7 @@ void FEdModeMeshPaint::OnVRAction( class FEditorViewportClient& ViewportClient,
 				// SetAllowTriggerLightPressLocking() function in the next tick (after VR Editor will have disabled locking.)
 				// @todo vreditor: We shouldn't have to use this trick.  This input system needs more flexibility in its public interface!
 				//bWasHandled = true;
-				//bOutIsInCaptured = true;
+				//bOutIsInputCaptured = true;
 
 				StartPainting();
 				PaintingWithInteractorInVR = Interactor;
@@ -5115,7 +5115,7 @@ void FEdModeMeshPaint::OnVRAction( class FEditorViewportClient& ViewportClient,
 
 				// @todo vreditor: We shouldn't have to use this trick.  This input system needs more flexibility in its public interface!
 				//bWasHandled = true;
-				//bOutIsInCaptured = false;
+				//bOutIsInputCaptured = false;
 			}
 		}
 	}
diff --git a/Engine/Source/Editor/MeshPaint/Public/MeshPaintEdMode.h b/Engine/Source/Editor/MeshPaint/Public/MeshPaintEdMode.h
--- a/Engine/Source/Editor/MeshPaint/Public/MeshPaintEdMode.h
+++ b/Engine/Source/Editor/MeshPaint/Public/MeshPaintEdMode.h
@@ -716,7 +716,7 @@ private:
 
 	/** Called when the user presses a button on their motion controller device */
 	void OnVRAction( class FEditorViewportClient& ViewportClient, class UViewportInteractor* Interactor,
-		const struct FViewportActionKeyInput& Action, const EInputEvent Event, bool& bOutIsInCaptured, bool& bWasHandled );
+		const struct FViewportActionKeyInput& Action, const EInputEvent Event, bool& bOutIsInputCaptured, bool& bWasHandled );
 
 private:
 
diff --git a/Engine/Source/Editor/VREditor/Teleporter/VREditorTeleporter.cpp b/Engine/Source/Editor/VREditor/Teleporter/VREditorTeleporter.cpp
--- a/Engine/Source/Editor/VREditor/Teleporter/VREditorTeleporter.cpp
+++ b/Engine/Source/Editor/VREditor/Teleporter/VREditorTeleporter.cpp
@@ -89,7 +89,7 @@ void UVREditorTeleporter::StartTeleport( UViewportInteractor* Interactor )
 	}
 }
 
-void UVREditorTeleporter::OnInteractorAction( class FEditorViewportClient& ViewportClient, UViewportInteractor* Interactor, const struct FViewportActionKeyInput& Action, const EInputEvent Event, bool& bOutIsInCaptured, bool& bWasHandled )
+void UVREditorTeleporter::OnInteractorAction( class FEditorViewportClient& ViewportClient, UViewportInteractor* Interactor, const struct FViewportActionKeyInput& Action, const EInputEvent Event, bool& bOutIsInputCaptured, bool& bWasHandled )
 {
 	if ( ( Action.ActionType == ViewportWorldActionTypes::SelectAndMove ||
 		Action.ActionType == ViewportWorldActionTypes::SelectAndMove_LightlyPressed ) )
diff --git a/Engine/Source/Editor/VREditor/Teleporter/VREditorTeleporter.h b/Engine/Source/Editor/VREditor/Teleporter/VREditorTeleporter.h
--- a/Engine/Source/Editor/VREditor/Teleporter/VREditorTeleporter.h
+++ b/Engine/Source/Editor/VREditor/Teleporter/VREditorTeleporter.h
@@ -36,7 +36,7 @@ private:
 
 	/** Called when the user presses a button on their motion controller device */
 	void OnInteractorAction( class FEditorViewportClient& ViewportClient, UViewportInteractor* Interactor,
-		const struct FViewportActionKeyInput& Action, const EInputEvent Event, bool& bOutIsInCaptured, bool& bWasHandled );
+		const struct FViewportActionKeyInput& Action, const EInputEvent Event, bool& bOutIsInputCaptured, bool& bWasHandled );
 
 	/** Actually teleport to a position */
 	void Teleport( const float DeltaTime );
diff --git a/Engine/Source/Editor/VREditor/VREditor.Build.cs b/Engine/Source/Editor/VREditor/VREditor.Build.cs
--- a/Engine/Source/Editor/VREditor/VREditor.Build.cs
+++ b/Engine/Source/Editor/VREditor/VREditor.Build.cs
@@ -19,7 +19,8 @@ namespace UnrealBuildTool.Rules
 					"UMG",
 					"LevelEditor",
 					"HeadMountedDisplay",
-                    "ViewportInteraction"
+                    "ViewportInteraction",
+					"Analytics"
 				}
 			);
 
diff --git a/Engine/Source/Editor/VREditor/VREditorInteractor.cpp b/Engine/Source/Editor/VREditor/VREditorInteractor.cpp
--- a/Engine/Source/Editor/VREditor/VREditorInteractor.cpp
+++ b/Engine/Source/Editor/VREditor/VREditorInteractor.cpp
@@ -1095,7 +1095,6 @@ void UVREditorInteractor::SetLaserVisuals( const FLinearColor& NewColor, const f
 void UVREditorInteractor::UpdateRadialMenuInput( const float DeltaTime )
 {
 	UVREditorUISystem& UISystem = GetVRMode().GetUISystem();
-	UVREditorInteractor* OtherInteractor = Cast<UVREditorInteractor>( GetOtherInteractor() );
 
 	//Update the radial menu
 	EViewportInteractionDraggingMode DraggingMode = GetDraggingMode();
@@ -1106,7 +1105,7 @@ void UVREditorInteractor::UpdateRadialMenuInput( const float DeltaTime )
 		DraggingMode != EViewportInteractionDraggingMode::ActorsAtLaserImpact &&
 		DraggingMode != EViewportInteractionDraggingMode::AssistingDrag && 
 		!UISystem.IsInteractorDraggingDockUI( this ) &&
-		!UISystem.IsShowingRadialMenu( OtherInteractor ) )
+		!UISystem.IsShowingRadialMenu( Cast<UVREditorInteractor>( OtherInteractor ) ) )
 	{
 		const EHMDDeviceType::Type HMDDevice = GetHMDDeviceType();
 
diff --git a/Engine/Source/Editor/VREditor/VREditorUISystem.cpp b/Engine/Source/Editor/VREditor/VREditorUISystem.cpp
--- a/Engine/Source/Editor/VREditor/VREditorUISystem.cpp
+++ b/Engine/Source/Editor/VREditor/VREditorUISystem.cpp
@@ -168,7 +168,7 @@ void UVREditorUISystem::Shutdown()
 }
 
 void UVREditorUISystem::OnVRAction( FEditorViewportClient& ViewportClient, UViewportInteractor* Interactor,
-	const FViewportActionKeyInput& Action, const EInputEvent Event, bool& bOutIsInCaptured, bool& bWasHandled )
+	const FViewportActionKeyInput& Action, const EInputEvent Event, bool& bOutIsInputCaptured, bool& bWasHandled )
 {
 	UVREditorInteractor* VREditorInteractor = Cast<UVREditorInteractor>( Interactor );
 	if ( VREditorInteractor )
@@ -270,7 +270,7 @@ void UVREditorUISystem::OnVRAction( FEditorViewportClient& ViewportClient, UView
 										VREditorInteractor->SetIsClickingOnUI( true );
 										VREditorInteractor->SetIsRightClickingOnUI( bIsRightClicking );
 										VREditorInteractor->SetLastClickReleaseTime( CurrentTime );
-										bOutIsInCaptured = true;
+										bOutIsInputCaptured = true;
 
 										// Play a haptic effect on press
 										VREditorInteractor->PlayHapticEffect( VREd::UIPressHapticFeedbackStrength->GetFloat() );
@@ -296,7 +296,7 @@ void UVREditorUISystem::OnVRAction( FEditorViewportClient& ViewportClient, UView
 						}
 						VREditorInteractor->SetIsClickingOnUI( false );
 						VREditorInteractor->SetIsRightClickingOnUI( false );
-						bOutIsInCaptured = false;
+						bOutIsInputCaptured = false;
 					}
 
 					if ( !bWasHandled )
diff --git a/Engine/Source/Editor/VREditor/VREditorUISystem.h b/Engine/Source/Editor/VREditor/VREditorUISystem.h
--- a/Engine/Source/Editor/VREditor/VREditorUISystem.h
+++ b/Engine/Source/Editor/VREditor/VREditorUISystem.h
@@ -124,7 +124,7 @@ protected:
 
 	/** Called when the user presses a button on their motion controller device */
 	void OnVRAction( class FEditorViewportClient& ViewportClient, UViewportInteractor* Interactor,
-		const struct FViewportActionKeyInput& Action, const EInputEvent Event, bool& bOutIsInCaptured, bool& bWasHandled );
+		const struct FViewportActionKeyInput& Action, const EInputEvent Event, bool& bOutIsInputCaptured, bool& bWasHandled );
 
 	/** Called every frame to update hover state */
 	void OnVRHoverUpdate( class FEditorViewportClient& ViewportClient, UViewportInteractor* Interactor, FVector& HoverImpactPoint, bool& bWasHandled );
diff --git a/Engine/Source/Editor/VREditor/VREditorWorldInteraction.cpp b/Engine/Source/Editor/VREditor/VREditorWorldInteraction.cpp
--- a/Engine/Source/Editor/VREditor/VREditorWorldInteraction.cpp
+++ b/Engine/Source/Editor/VREditor/VREditorWorldInteraction.cpp
@@ -279,9 +279,9 @@ void UVREditorWorldInteraction::OnAssetDragStartedFromContentBrowser( const TArr
 
 				// Attempt to create actors from the dropped object
 				const bool bSelectNewActors = true;
-				const EObjectFlags ObjectFlags = bIsPreview ? RF_Transient : RF_Transactional;
+				const EObjectFlags NewObjectFlags = bIsPreview ? RF_Transient : RF_Transactional;
 
-				TArray<AActor*> NewActors = FLevelEditorViewportClient::TryPlacingActorFromObject( GetViewportWorld()->GetCurrentLevel(), AssetObj, bSelectNewActors, ObjectFlags, FactoryToUse );
+				TArray<AActor*> NewActors = FLevelEditorViewportClient::TryPlacingActorFromObject( GetViewportWorld()->GetCurrentLevel(), AssetObj, bSelectNewActors, NewObjectFlags, FactoryToUse );
 
 				if( NewActors.Num() > 0 )
 				{
diff --git a/Engine/Source/Editor/ViewportInteraction/Public/ViewportInteractor.h b/Engine/Source/Editor/ViewportInteraction/Public/ViewportInteractor.h
--- a/Engine/Source/Editor/ViewportInteraction/Public/ViewportInteractor.h
+++ b/Engine/Source/Editor/ViewportInteraction/Public/ViewportInteractor.h
@@ -135,6 +135,7 @@ public:
 	virtual bool GetIsLaserBlocked() { return false; }
 
 	/** Gets a certain action by iterating through the map looking for the same ActionType */
+	// @todo ViewportInteractor: This should be changed to return a const pointer, but we need to fix up some dragging code in WorldInteraction first
 	FViewportActionKeyInput* GetActionWithName( const FName InActionName );
 
 	/** Gets the drag haptic feedback strength console variable */
diff --git a/Engine/Source/Editor/ViewportInteraction/Public/ViewportWorldInteractionInterface.h b/Engine/Source/Editor/ViewportInteraction/Public/ViewportWorldInteractionInterface.h
--- a/Engine/Source/Editor/ViewportInteraction/Public/ViewportWorldInteractionInterface.h
+++ b/Engine/Source/Editor/ViewportInteraction/Public/ViewportWorldInteractionInterface.h
@@ -36,6 +36,6 @@ public:
 
 	/** Gets the event for input actions update which is broadcasted for each interactor */
 	DECLARE_EVENT_SixParams( IViewportWorldInteractionInterface, FOnVIActionHandle, class FEditorViewportClient& /* ViewportClient */, class UViewportInteractor* /* Interactor */, 
-	const struct FViewportActionKeyInput& /* Action */, const EInputEvent /* Event */, bool& /* bOutIsInCaptured */, bool& /* bWasHandled */ );
+	const struct FViewportActionKeyInput& /* Action */, const EInputEvent /* Event */, bool& /* bOutIsInputCaptured */, bool& /* bWasHandled */ );
 	virtual FOnVIActionHandle& OnViewportInteractionInputAction() = 0;
 };
\ No newline at end of file
diff --git a/Engine/Source/Editor/ViewportInteraction/ViewportInteractor.cpp b/Engine/Source/Editor/ViewportInteraction/ViewportInteractor.cpp
--- a/Engine/Source/Editor/ViewportInteraction/ViewportInteractor.cpp
+++ b/Engine/Source/Editor/ViewportInteraction/ViewportInteractor.cpp
@@ -308,6 +308,7 @@ bool UViewportInteractor::HandleInputKey( const FKey Key, const EInputEvent Even
 						!bIsDraggingWorldWithTwoHands )
 					{
 						WorldInteraction->StopDragging( this );
+						Action->bIsInputCaptured = false;
 					}
 
 					bHandled = true;
@@ -359,6 +360,8 @@ bool UViewportInteractor::HandleInputKey( const FKey Key, const EInputEvent Even
 				else if ( Event == IE_Released )
 				{
 					WorldInteraction->StopDragging( this );
+
+					Action->bIsInputCaptured = false;
 				}
 			}
 			else if ( Action->ActionType == ViewportWorldActionTypes::Delete )
diff --git a/Engine/Source/Editor/ViewportInteraction/ViewportWorldInteraction.cpp b/Engine/Source/Editor/ViewportInteraction/ViewportWorldInteraction.cpp
--- a/Engine/Source/Editor/ViewportInteraction/ViewportWorldInteraction.cpp
+++ b/Engine/Source/Editor/ViewportInteraction/ViewportWorldInteraction.cpp
@@ -2217,12 +2217,6 @@ void UViewportWorldInteraction::StopDragging( UViewportInteractor* Interactor )
 		}
 		
 		InteractorData.DraggingMode = EViewportInteractionDraggingMode::Nothing;
-		
-		FViewportActionKeyInput* OptionalAction = Interactor->GetActionWithName( ViewportWorldActionTypes::SelectAndMove );
-		if ( OptionalAction )
-		{
-			OptionalAction->bIsInputCaptured = false;
-		}
 	}
 
 	// NOTE: Even though we're done dragging, we keep our list of transformables so that inertial drag can still push them around!
@@ -2524,22 +2518,22 @@ void UViewportWorldInteraction::RefreshTransformGizmo( const bool bNewObjectsSel
 
 			{
 				// Update snap grid material state
-				UMaterialInstanceDynamic* SnapGridMID = GetSnapGridMID();
+				UMaterialInstanceDynamic* LocalSnapGridMID = GetSnapGridMID();
 
 				static FName GridColorParameterName( "GridColor" );
-				SnapGridMID->SetVectorParameterValue( GridColorParameterName, GridColor );
+				LocalSnapGridMID->SetVectorParameterValue( GridColorParameterName, GridColor );
 
 				static FName GridCenterParameterName( "GridCenter" );
-				SnapGridMID->SetVectorParameterValue( GridCenterParameterName, SnapGridCenter );
+				LocalSnapGridMID->SetVectorParameterValue( GridCenterParameterName, SnapGridCenter );
 
 				static FName GridIntervalParameterName( "GridInterval" );
-				SnapGridMID->SetScalarParameterValue( GridIntervalParameterName, GridInterval );
+				LocalSnapGridMID->SetScalarParameterValue( GridIntervalParameterName, GridInterval );
 
 				static FName GridRadiusParameterName( "GridRadius" );
-				SnapGridMID->SetScalarParameterValue( GridRadiusParameterName, SnapGridSize * 0.5f );
+				LocalSnapGridMID->SetScalarParameterValue( GridRadiusParameterName, SnapGridSize * 0.5f );
 
 				static FName LineWidthParameterName( "LineWidth" );
-				SnapGridMID->SetScalarParameterValue( LineWidthParameterName, LineWidth );
+				LocalSnapGridMID->SetScalarParameterValue( LineWidthParameterName, LineWidth );
 
 				FMatrix GridRotationMatrix = GizmoToWorld.ToMatrixNoScale().Inverse();
 				GridRotationMatrix.SetOrigin( FVector::ZeroVector );
@@ -2547,9 +2541,9 @@ void UViewportWorldInteraction::RefreshTransformGizmo( const bool bNewObjectsSel
 				static FName GridRotationMatrixXParameterName( "GridRotationMatrixX" );
 				static FName GridRotationMatrixYParameterName( "GridRotationMatrixY" );
 				static FName GridRotationMatrixZParameterName( "GridRotationMatrixZ" );
-				SnapGridMID->SetVectorParameterValue( GridRotationMatrixXParameterName, GridRotationMatrix.GetScaledAxis( EAxis::X ) );
-				SnapGridMID->SetVectorParameterValue( GridRotationMatrixYParameterName, GridRotationMatrix.GetScaledAxis( EAxis::Y ) );
-				SnapGridMID->SetVectorParameterValue( GridRotationMatrixZParameterName, GridRotationMatrix.GetScaledAxis( EAxis::Z ) );
+				LocalSnapGridMID->SetVectorParameterValue( GridRotationMatrixXParameterName, GridRotationMatrix.GetScaledAxis( EAxis::X ) );
+				LocalSnapGridMID->SetVectorParameterValue( GridRotationMatrixYParameterName, GridRotationMatrix.GetScaledAxis( EAxis::Y ) );
+				LocalSnapGridMID->SetVectorParameterValue( GridRotationMatrixZParameterName, GridRotationMatrix.GetScaledAxis( EAxis::Z ) );
 			}
 		}
 		else
