Repplying "Undoing Duplication" fix #jira UE-27727

diff --git a/Engine/Source/Editor/ViewportInteraction/ViewportWorldInteraction.cpp b/Engine/Source/Editor/ViewportInteraction/ViewportWorldInteraction.cpp
--- a/Engine/Source/Editor/ViewportInteraction/ViewportWorldInteraction.cpp
+++ b/Engine/Source/Editor/ViewportInteraction/ViewportWorldInteraction.cpp
@@ -471,9 +471,8 @@ void UViewportWorldInteraction::Paste()
 
 void UViewportWorldInteraction::Duplicate()
 {
-	ABrush::SetSuppressBSPRegeneration( true );
-	GEditor->edactDuplicateSelected( GetViewportWorld()->GetCurrentLevel(), false );
-	ABrush::SetSuppressBSPRegeneration( false );
+	// @todo vreditor: Needs CanExecute()  (see LevelEditorActions.cpp)
+	GUnrealEd->Exec( GetViewportWorld(), TEXT( "DUPLICATE" ) );
 }
 
 void UViewportWorldInteraction::Deselect()
