VR Editor: Motion controllers are now enabled in forced VR mode (merged from CL 3027012)
(Console command is 'VREd.ForceVRMode')

diff --git a/Engine/Source/Editor/VREditor/VREditorInteractor.cpp b/Engine/Source/Editor/VREditor/VREditorInteractor.cpp
--- a/Engine/Source/Editor/VREditor/VREditorInteractor.cpp
+++ b/Engine/Source/Editor/VREditor/VREditorInteractor.cpp
@@ -267,9 +267,8 @@ void UVREditorInteractor::Tick( const float DeltaTime )
 	{
 		const float WorldScaleFactor = WorldInteraction->GetWorldScaleFactor();
 
-		// We don't bother drawing hand meshes if we're in "forced VR mode" (because they'll just be on top of the camera).
-		// Also, don't bother drawing hands if we're not currently tracking them.
-		if ( GetVRMode().IsActuallyUsingVR() && bHaveMotionController )
+		// Don't bother drawing hands if we're not currently tracking them.
+		if ( bHaveMotionController )
 		{
 			HandMeshComponent->SetVisibility( true );
 		}
@@ -281,9 +280,6 @@ void UVREditorInteractor::Tick( const float DeltaTime )
 		FVector LaserPointerStart, LaserPointerEnd;
 		if ( GetLaserPointer( /* Out */ LaserPointerStart, /* Out */ LaserPointerEnd ) )
 		{
-			// Only show the laser if we're actually in VR
-			LaserPointerMeshComponent->SetVisibility( GetVRMode().IsActuallyUsingVR() );
-
 			// NOTE: We don't need to set the laser pointer location and rotation, as the MotionControllerComponent will do
 			// that later in the frame.  
 
diff --git a/Engine/Source/Editor/VREditor/VREditorMode.cpp b/Engine/Source/Editor/VREditor/VREditorMode.cpp
--- a/Engine/Source/Editor/VREditor/VREditorMode.cpp
+++ b/Engine/Source/Editor/VREditor/VREditorMode.cpp
@@ -329,7 +329,6 @@ void FVREditorMode::Enter()
 		// Setup world interaction
 		WorldInteraction = NewObject<UVREditorWorldInteraction>();
 		WorldInteraction->SetOwner( this );
-		WorldInteraction->SetIsUsingVR( bActuallyUsingVR );
 		WorldInteraction->Init( VREditorLevelViewportWeakPtr.Pin()->GetViewportClient().Get() );
 
 		// Motion controllers
@@ -364,18 +363,6 @@ void FVREditorMode::Enter()
 
 	}
 
-	//// If we're not actually in VR mode, and instead we've forced it, we should spawn some UI at the Origin to help test
-	//if ( !bActuallyUsingVR )
-	//{
-	//	const int32 HandIndex = 1;
-	//	const bool bShouldShow = true;
-
-	//	FVirtualHand& Hand = VirtualHands[HandIndex];
-	//	Hand.Transform = FTransform(FVector(0, 0, 0));
-
-	//	UISystem->ShowEditorUIPanel(FVREditorUISystem::EEditorUIPanel::ContentBrowser, HandIndex, bShouldShow);
-	//}
-
 	// Call parent implementation
 	FEdMode::Enter();
 
diff --git a/Engine/Source/Editor/ViewportInteraction/Public/ViewportWorldInteraction.h b/Engine/Source/Editor/ViewportInteraction/Public/ViewportWorldInteraction.h
--- a/Engine/Source/Editor/ViewportInteraction/Public/ViewportWorldInteraction.h
+++ b/Engine/Source/Editor/ViewportInteraction/Public/ViewportWorldInteraction.h
@@ -437,12 +437,6 @@ private:
 	/** The current Gizmo type that is used for the TransformGizmo Actor */
 	EGizmoHandleTypes CurrentGizmoType;
 
-	//
-	// VR
-	//
-
-	/** If we are in VR right now */
-	bool bIsUsingVR;
 
 	//
 	// Snap grid
diff --git a/Engine/Source/Editor/ViewportInteraction/ViewportWorldInteraction.cpp b/Engine/Source/Editor/ViewportInteraction/ViewportWorldInteraction.cpp
--- a/Engine/Source/Editor/ViewportInteraction/ViewportWorldInteraction.cpp
+++ b/Engine/Source/Editor/ViewportInteraction/ViewportWorldInteraction.cpp
@@ -3,6 +3,7 @@
 #include "ViewportInteractionModule.h"
 #include "ViewportWorldInteraction.h"
 #include "ViewportInteractor.h"
+#include "MouseCursorInteractor.h"
 #include "ViewportInteractableInterface.h"
 #include "ViewportDragOperation.h"
 #include "VIGizmoHandle.h"
@@ -227,7 +228,6 @@ UViewportWorldInteraction::UViewportWorldInteraction( const FObjectInitializer&
 	StartDragAngleOnRotation( ),
 	StartDragHandleDirection( ),
 	CurrentGizmoType( EGizmoHandleTypes::All ),
-	bIsUsingVR( false ),
 	SnapGridActor( nullptr ),
 	SnapGridMeshComponent( nullptr ),
 	SnapGridMID( nullptr ),
@@ -389,7 +389,7 @@ FTransform UViewportWorldInteraction::GetRoomTransform() const
 FTransform UViewportWorldInteraction::GetRoomSpaceHeadTransform() const
 {
 	FTransform HeadTransform = FTransform::Identity;
-	if ( bIsUsingVR && GEngine->HMDDevice.IsValid() )
+	if ( GEngine->HMDDevice.IsValid() )
 	{
 		FQuat RoomSpaceHeadOrientation;
 		FVector RoomSpaceHeadLocation;
@@ -435,11 +435,6 @@ FEditorViewportClient* UViewportWorldInteraction::GetViewportClient() const
 	return EditorViewportClient;
 }
 
-void UViewportWorldInteraction::SetIsUsingVR( const bool bInIsUsingVR )
-{
-	bIsUsingVR = bInIsUsingVR;
-}
-
 void UViewportWorldInteraction::Undo()
 {
 	GUnrealEd->Exec( GetViewportWorld(), TEXT( "TRANSACTION UNDO" ) );
@@ -812,12 +807,13 @@ void UViewportWorldInteraction::InteractionTick( FEditorViewportClient* Viewport
 			}
 
 			{
-				// Intertia must only work in VR
+				const bool bIsMouseCursorInteractor = Cast<UMouseCursorInteractor>( Interactor ) != nullptr;
 				// @todo vreditor: We could do a world space distance test (times world scale factor) when forcing VR mode to get similar (but not quite the same) behavior
-				if( bIsUsingVR )
 				{
 					// Don't bother with inertia if we're not moving very fast.  This filters out tiny accidental movements.
-					const FVector RoomSpaceHandDelta = ( InteractorData.RoomSpaceTransform.GetLocation() - InteractorData.LastRoomSpaceTransform.GetLocation() );
+					const FVector RoomSpaceHandDelta = bIsMouseCursorInteractor ?
+						( InteractorData.Transform.GetLocation() - InteractorData.LastTransform.GetLocation() ) :
+						( InteractorData.RoomSpaceTransform.GetLocation() - InteractorData.LastRoomSpaceTransform.GetLocation() );
 					if( RoomSpaceHandDelta.Size() < VI::MinVelocityForInertia->GetFloat() * WorldScaleFactor )
 					{
 						InteractorData.DragTranslationVelocity = FVector::ZeroVector;
