[NO NON-PII COMMIT MESSAGE]

diff --git a/Engine/Source/Editor/ViewportInteraction/Public/ViewportWorldInteraction.h b/Engine/Source/Editor/ViewportInteraction/Public/ViewportWorldInteraction.h
--- a/Engine/Source/Editor/ViewportInteraction/Public/ViewportWorldInteraction.h
+++ b/Engine/Source/Editor/ViewportInteraction/Public/ViewportWorldInteraction.h
@@ -253,10 +253,21 @@ private:
 		FVector& GizmoSpaceFirstDragUpdateOffsetAlongAxis,
 		FVector& DragDeltaFromStartOffset,
 		bool& bIsDrivingVelocityOfSimulatedTransformables,
-		FVector& OutUnsnappedDraggedTo );
+		FVector& OutUnsnappedDraggedTo);
 	
 	/** Given a drag delta from a starting point, contrains that delta based on a gizmo handle axis */
-	FVector ComputeConstrainedDragDeltaFromStart( const bool bIsFirstDragUpdate, const ETransformGizmoInteractionType InteractionType, const TOptional<FTransformGizmoHandlePlacement> OptionalHandlePlacement, const FVector& DragDeltaFromStart, const FVector& LaserPointerStart, const FVector& LaserPointerDirection, const bool bIsLaserPointerValid, const FTransform& GizmoStartTransform, FVector& GizmoSpaceFirstDragUpdateOffsetAlongAxis, FVector& DragDeltaFromStartOffset ) const;
+	FVector ComputeConstrainedDragDeltaFromStart( 
+		const bool bIsFirstDragUpdate, 
+		const ETransformGizmoInteractionType InteractionType, 
+		const TOptional<FTransformGizmoHandlePlacement> OptionalHandlePlacement, 
+		const FVector& DragDeltaFromStart, 
+		const FVector& LaserPointerStart, 
+		const FVector& LaserPointerDirection, 
+		const bool bIsLaserPointerValid, 
+		const FTransform& GizmoStartTransform, 
+		FVector& GizmoSpaceFirstDragUpdateOffsetAlongAxis, 
+		FVector& DragDeltaFromStartOffset,
+		FVector& OutClosestPointOnLaser) const;
 
 	/** Checks to see if smooth snapping is currently enabled */
 	bool IsSmoothSnappingEnabled() const;
diff --git a/Engine/Source/Editor/ViewportInteraction/ViewportWorldInteraction.cpp b/Engine/Source/Editor/ViewportInteraction/ViewportWorldInteraction.cpp
--- a/Engine/Source/Editor/ViewportInteraction/ViewportWorldInteraction.cpp
+++ b/Engine/Source/Editor/ViewportInteraction/ViewportWorldInteraction.cpp
@@ -1278,6 +1278,7 @@ void UViewportWorldInteraction::UpdateDragging(
 	if ( DraggingMode == EViewportInteractionDraggingMode::ActorsWithGizmo ||
 		DraggingMode == EViewportInteractionDraggingMode::ActorsAtLaserImpact )
 	{
+		FVector OutClosestPointOnLaser;
 		FVector ConstrainedDragDeltaFromStart = ComputeConstrainedDragDeltaFromStart(
 			bIsFirstDragUpdate,
 			InteractionType,
@@ -1288,13 +1289,16 @@ void UViewportWorldInteraction::UpdateDragging(
 			bIsLaserPointerValid,
 			GizmoStartTransform,
 			/* In/Out */ GizmoSpaceFirstDragUpdateOffsetAlongAxis,
-			/* In/Out */ DragDeltaFromStartOffset );
+			/* In/Out */ DragDeltaFromStartOffset,
+			/* Out */ OutClosestPointOnLaser);
 
-		FVector UnsnappedDraggedTo = GizmoStartTransform.GetLocation() + ConstrainedDragDeltaFromStart;
-		OutUnsnappedDraggedTo = UnsnappedDraggedTo;
+		ConstrainedDragDeltaFromStart = GizmoStartTransform.GetLocation() + ConstrainedDragDeltaFromStart;
+
+		// Set out put for hover point
+		OutUnsnappedDraggedTo = GizmoStartTransform.GetLocation() + OutClosestPointOnLaser;
 
 		// Grid snap!
-		FVector SnappedDraggedTo = UnsnappedDraggedTo;
+		FVector SnappedDraggedTo = ConstrainedDragDeltaFromStart;
 		if ( FSnappingUtils::IsSnapToGridEnabled() )
 		{
 			// Snap in local space, if we need to
@@ -1314,7 +1318,7 @@ void UViewportWorldInteraction::UpdateDragging(
 		for ( int32 PassIndex = 0; PassIndex < 2; ++PassIndex )
 		{
 			const bool bIsUpdatingUnsnappedTarget = ( PassIndex == 1 );
-			const FVector& PassDraggedTo = bIsUpdatingUnsnappedTarget ? UnsnappedDraggedTo : SnappedDraggedTo;
+			const FVector& PassDraggedTo = bIsUpdatingUnsnappedTarget ? ConstrainedDragDeltaFromStart : SnappedDraggedTo;
 
 			if ( InteractionType == ETransformGizmoInteractionType::Translate ||
 				InteractionType == ETransformGizmoInteractionType::TranslateOnPlane )
@@ -1562,12 +1566,12 @@ void UViewportWorldInteraction::UpdateDragging(
 						}
 
 						// Get the laserpointer intersection on the plane of the handle
-						FPlane Plane = FPlane( RotationHandlePlaneLocation, StartDragHandleDirection.GetValue() );
+						FPlane RotationHandlePlane = FPlane( RotationHandlePlaneLocation, StartDragHandleDirection.GetValue() );
 						
-						const FVector IntersectWorldLocationOnPlane = FMath::LinePlaneIntersection( LaserPointerStart, LaserPointerStart + LaserPointerDirection, 
-							RotationHandlePlaneLocation, StartDragHandleDirection.GetValue() );
+						const FVector IntersectWorldLocationOnPlane = FMath::LinePlaneIntersection( LaserPointerStart, LaserPointerStart + LaserPointerDirection, RotationHandlePlane);
 
-						DrawDebugLine( GetViewportWorld(), TransformGizmoActor->GetActorLocation(), IntersectWorldLocationOnPlane, FColor::Red, false, -1.f, 0, 3.0f * GetWorldScaleFactor() );
+						// Set output for hover point
+						OutUnsnappedDraggedTo = IntersectWorldLocationOnPlane;
 
 						// The starting transformation of this transformable in gizmo space
 						const FTransform GizmoSpaceStartTransform = Transformable.StartTransform * WorldToGizmo;
@@ -1925,7 +1929,18 @@ void UViewportWorldInteraction::UpdateDragging(
 	bIsFirstDragUpdate = false;
 }
 
-FVector UViewportWorldInteraction::ComputeConstrainedDragDeltaFromStart( const bool bIsFirstDragUpdate, const ETransformGizmoInteractionType InteractionType, const TOptional<FTransformGizmoHandlePlacement> OptionalHandlePlacement, const FVector& DragDeltaFromStart, const FVector& LaserPointerStart, const FVector& LaserPointerDirection, const bool bIsLaserPointerValid, const FTransform& GizmoStartTransform, FVector& GizmoSpaceFirstDragUpdateOffsetAlongAxis, FVector& GizmoSpaceDragDeltaFromStartOffset ) const
+FVector UViewportWorldInteraction::ComputeConstrainedDragDeltaFromStart( 
+	const bool bIsFirstDragUpdate, 
+	const ETransformGizmoInteractionType InteractionType, 
+	const TOptional<FTransformGizmoHandlePlacement> OptionalHandlePlacement, 
+	const FVector& DragDeltaFromStart, 
+	const FVector& LaserPointerStart, 
+	const FVector& LaserPointerDirection, 
+	const bool bIsLaserPointerValid, 
+	const FTransform& GizmoStartTransform, 
+	FVector& GizmoSpaceFirstDragUpdateOffsetAlongAxis, 
+	FVector& GizmoSpaceDragDeltaFromStartOffset,
+	FVector& OutClosestPointOnLaser) const
 {
 	FVector ConstrainedGizmoSpaceDeltaFromStart;
 
@@ -1974,22 +1989,28 @@ FVector UViewportWorldInteraction::ComputeConstrainedDragDeltaFromStart( const b
 				// Gizmo space is defined as the starting position of the interaction, so we simply take the closest position
 				// along the plane as our delta from the start in gizmo space
 				ConstrainedGizmoSpaceDeltaFromStart = GizmoSpacePointOnPlane;
+
+				// Set out for the closest point on laser for clipping
+				OutClosestPointOnLaser = GizmoStartTransform.TransformVectorNoScale(GizmoSpacePointOnPlane);
 			}
 			else
 			{
 				// 1D Line
 				const float AxisSegmentLength = LaserPointerMaxLength * 2.0f;	// @todo vreditor: We're creating a line segment to represent an infinitely long axis, but really it just needs to be further than our laser pointer can reach
 
-				DVector GizmoSpaceClosestPointOnRayDouble, GizmoSpaceClosestPointOnAxisDouble;
+				DVector GizmoSpaceClosestPointOnLaserDouble, GizmoSpaceClosestPointOnAxisDouble;
 				SegmentDistToSegmentDouble(
 					GizmoSpaceLaserPointerStart, DVector( GizmoSpaceLaserPointerStart ) + DVector( GizmoSpaceLaserPointerDirection ) * LaserPointerMaxLength,	// @todo vreditor: Should store laser pointer length rather than hard code
 					DVector( GizmoSpaceConstraintAxis ) * -AxisSegmentLength, DVector( GizmoSpaceConstraintAxis ) * AxisSegmentLength,
-					/* Out */ GizmoSpaceClosestPointOnRayDouble,
+					/* Out */ GizmoSpaceClosestPointOnLaserDouble,
 					/* Out */ GizmoSpaceClosestPointOnAxisDouble );
 
 				// Gizmo space is defined as the starting position of the interaction, so we simply take the closest position
 				// along the axis as our delta from the start in gizmo space
 				ConstrainedGizmoSpaceDeltaFromStart = GizmoSpaceClosestPointOnAxisDouble.ToFVector();
+
+				// Set out for the closest point on laser for clipping
+				OutClosestPointOnLaser = GizmoStartTransform.TransformVectorNoScale(GizmoSpaceClosestPointOnLaserDouble.ToFVector());
 			}
 
 			// Account for the handle position on the outside of the bounds.  This is a bit confusing but very important for dragging
