Fixed world scaling for hover point while transforming in VR editor

diff --git a/Engine/Source/Editor/ViewportInteraction/ViewportWorldInteraction.cpp b/Engine/Source/Editor/ViewportInteraction/ViewportWorldInteraction.cpp
--- a/Engine/Source/Editor/ViewportInteraction/ViewportWorldInteraction.cpp
+++ b/Engine/Source/Editor/ViewportInteraction/ViewportWorldInteraction.cpp
@@ -1991,7 +1991,7 @@ FVector UViewportWorldInteraction::ComputeConstrainedDragDeltaFromStart(
 				ConstrainedGizmoSpaceDeltaFromStart = GizmoSpacePointOnPlane;
 
 				// Set out for the closest point on laser for clipping
-				OutClosestPointOnLaser = GizmoStartTransform.TransformVectorNoScale(GizmoSpacePointOnPlane);
+				OutClosestPointOnLaser = GizmoStartTransform.TransformVector(GizmoSpacePointOnPlane);
 			}
 			else
 			{
@@ -2010,7 +2010,7 @@ FVector UViewportWorldInteraction::ComputeConstrainedDragDeltaFromStart(
 				ConstrainedGizmoSpaceDeltaFromStart = GizmoSpaceClosestPointOnAxisDouble.ToFVector();
 
 				// Set out for the closest point on laser for clipping
-				OutClosestPointOnLaser = GizmoStartTransform.TransformVectorNoScale(GizmoSpaceClosestPointOnLaserDouble.ToFVector());
+				OutClosestPointOnLaser = GizmoStartTransform.TransformVector(GizmoSpaceClosestPointOnLaserDouble.ToFVector());
 			}
 
 			// Account for the handle position on the outside of the bounds.  This is a bit confusing but very important for dragging
