VR Editor: Fixed VR Editor assertion in server builds

diff --git a/Engine/Source/Editor/VREditor/VREditorMotionControllerInteractor.cpp b/Engine/Source/Editor/VREditor/VREditorMotionControllerInteractor.cpp
--- a/Engine/Source/Editor/VREditor/VREditorMotionControllerInteractor.cpp
+++ b/Engine/Source/Editor/VREditor/VREditorMotionControllerInteractor.cpp
@@ -40,7 +40,7 @@ namespace VREd
 	static FAutoConsoleVariable TrackpadRelativeDragSpeed( TEXT( "VREd.TrackpadRelativeDragSpeed" ), 8.0f, TEXT( "How fast objects move toward or away when you hold a direction on an analog stick while carrying them" ) );
 	static FAutoConsoleVariable InvertTrackpadVertical( TEXT( "VREd.InvertTrackpadVertical" ), 1, TEXT( "Toggles inverting the touch pad vertical axis" ) );
 	static FAutoConsoleVariable TriggerLightlyPressedLockTime( TEXT( "VREd.TriggerLightlyPressedLockTime" ), 0.15f, TEXT( "If the trigger remains lightly pressed for longer than this, we'll continue to treat it as a light press in some cases" ) );
-	static FAutoConsoleVariable MinVelocityForInertia( TEXT( "VREd.MinVelocityForInertia" ), 1.0f, TEXT( "Minimum velocity (in cm/frame in unscaled room space) before inertia will kick in when releasing objects (or the world)" ) );
+	static FAutoConsoleVariable MinVelocityForInertia( TEXT( "VREd.MinVelocityForMotionControllerInertia" ), 1.0f, TEXT( "Minimum velocity (in cm/frame in unscaled room space) before inertia will kick in when releasing objects (or the world)" ) );
 	static FAutoConsoleVariable MinJoystickOffsetBeforeRadialMenu( TEXT( "VREd.MinJoystickOffsetBeforeRadialMenu" ), 0.15f, TEXT( "Toggles inverting the touch pad vertical axis" ) );
 
 	static FAutoConsoleVariable HelpLabelFadeDuration( TEXT( "VREd.HelpLabelFadeDuration" ), 0.4f, TEXT( "Duration to fade controller help labels in and out" ) );
diff --git a/Engine/Source/Editor/ViewportInteraction/ViewportWorldInteraction.cpp b/Engine/Source/Editor/ViewportInteraction/ViewportWorldInteraction.cpp
--- a/Engine/Source/Editor/ViewportInteraction/ViewportWorldInteraction.cpp
+++ b/Engine/Source/Editor/ViewportInteraction/ViewportWorldInteraction.cpp
@@ -232,6 +232,10 @@ UViewportWorldInteraction::UViewportWorldInteraction( const FObjectInitializer&
 	SnapGridMeshComponent( nullptr ),
 	SnapGridMID( nullptr ),
 	DraggedInteractable( nullptr )
+{
+}
+
+void UViewportWorldInteraction::Init( FEditorViewportClient* InEditorViewportClient )
 {
 	InputProcessor = MakeShareable( new FViewportInteractionInputProcessor( *this ) );
 	FSlateApplication::Get().SetInputPreProcessor( true, InputProcessor );
@@ -239,20 +243,17 @@ UViewportWorldInteraction::UViewportWorldInteraction( const FObjectInitializer&
 	// Find out about selection changes
 	USelection::SelectionChangedEvent.AddUObject( this, &UViewportWorldInteraction::OnActorSelectionChanged ); //@todo viewportinteraction
 
-	// Setting up colors
-	Colors.SetNumZeroed( ( int32 ) EColors::TotalCount );
+																											   // Setting up colors
+	Colors.SetNumZeroed( (int32)EColors::TotalCount );
 	{
-		Colors[ ( int32 ) EColors::DefaultColor ] = FLinearColor( 0.7f, 0.7f, 0.7f, 1.0f );;
-		Colors[ ( int32 ) EColors::Forward ] = FLinearColor( 0.4f, 0.05f, 0.05f, 1.0f );
-		Colors[ ( int32 ) EColors::Right ] = FLinearColor( 0.05f, 0.4f, 0.05f, 1.0f );
-		Colors[ ( int32 ) EColors::Up ] = FLinearColor( 0.05f, 0.05f, 0.4f, 1.0f );
-		Colors[ ( int32 ) EColors::Hover ] = FLinearColor::Yellow;
-		Colors[ ( int32 ) EColors::Dragging ] = FLinearColor::White;
+		Colors[ (int32)EColors::DefaultColor ] = FLinearColor( 0.7f, 0.7f, 0.7f, 1.0f );;
+		Colors[ (int32)EColors::Forward ] = FLinearColor( 0.4f, 0.05f, 0.05f, 1.0f );
+		Colors[ (int32)EColors::Right ] = FLinearColor( 0.05f, 0.4f, 0.05f, 1.0f );
+		Colors[ (int32)EColors::Up ] = FLinearColor( 0.05f, 0.05f, 0.4f, 1.0f );
+		Colors[ (int32)EColors::Hover ] = FLinearColor::Yellow;
+		Colors[ (int32)EColors::Dragging ] = FLinearColor::White;
 	}
-}
 
-void UViewportWorldInteraction::Init( FEditorViewportClient* InEditorViewportClient )
-{
 	EditorViewportClient = InEditorViewportClient;
 	AppTimeEntered = FTimespan::FromSeconds( FApp::GetCurrentTime() );
 
