[INTEGRATE] Fixed non-unity and Linux editor compilation errors (from CL 3052833 in //UE4/Main)

diff --git a/Engine/Source/Editor/MeshPaint/Private/MeshPaintEdMode.cpp b/Engine/Source/Editor/MeshPaint/Private/MeshPaintEdMode.cpp
--- a/Engine/Source/Editor/MeshPaint/Private/MeshPaintEdMode.cpp
+++ b/Engine/Source/Editor/MeshPaint/Private/MeshPaintEdMode.cpp
@@ -511,7 +511,7 @@ bool FEdModeMeshPaint::InputKey( FEditorViewportClient* InViewportClient, FViewp
 
 		// Stop current tracking if the user is no longer painting
 		if( bIsPainting && !bUserWantsPaint &&
-			( InKey == EKeys::LeftMouseButton || InKey == EKeys::RightMouseButton && InKey == EKeys::LeftAlt || InKey == EKeys::RightAlt ) )
+			( InKey == EKeys::LeftMouseButton || InKey == EKeys::RightMouseButton || InKey == EKeys::LeftAlt || InKey == EKeys::RightAlt ) )
 		{
 			bHandled = true;
 			EndPainting();
diff --git a/Engine/Source/Editor/ViewportInteraction/Gizmo/VIGizmoHandle.cpp b/Engine/Source/Editor/ViewportInteraction/Gizmo/VIGizmoHandle.cpp
--- a/Engine/Source/Editor/ViewportInteraction/Gizmo/VIGizmoHandle.cpp
+++ b/Engine/Source/Editor/ViewportInteraction/Gizmo/VIGizmoHandle.cpp
@@ -356,8 +356,8 @@ void UAxisGizmoHandleGroup::UpdateHandlesRelativeTransformOnAxis( const FTransfo
 			for ( int32 AxisIndex = 0; AxisIndex < 3; ++AxisIndex )
 			{		
 				if ( AxisIndex != FacingAxisIndex &&
-					( GizmoSpaceFacingAxisOffset[ AxisIndex ] > 0 && GizmoSpaceViewLocation[ AxisIndex ] < 0 ) ||
-					( GizmoSpaceFacingAxisOffset[ AxisIndex ] < 0 && GizmoSpaceViewLocation[ AxisIndex ] > 0 ) )
+					( ( GizmoSpaceFacingAxisOffset[ AxisIndex ] > 0 && GizmoSpaceViewLocation[ AxisIndex ] < 0 ) ||
+					  ( GizmoSpaceFacingAxisOffset[ AxisIndex ] < 0 && GizmoSpaceViewLocation[ AxisIndex ] > 0 ) ) )
 				{
 					GizmoSpaceFacingAxisOffset[ AxisIndex ] *= -1.0f;
 				}	
diff --git a/Engine/Source/Editor/ViewportInteraction/Gizmo/VIStretchGizmoHandle.cpp b/Engine/Source/Editor/ViewportInteraction/Gizmo/VIStretchGizmoHandle.cpp
--- a/Engine/Source/Editor/ViewportInteraction/Gizmo/VIStretchGizmoHandle.cpp
+++ b/Engine/Source/Editor/ViewportInteraction/Gizmo/VIStretchGizmoHandle.cpp
@@ -1,7 +1,8 @@
 // Copyright 1998-2015 Epic Games, Inc. All Rights Reserved.
 
 #include "ViewportInteractionModule.h"
-#include "VIGizmoHandle.h"
+#include "VIStretchGizmoHandle.h"
+#include "VIBaseTransformGizmo.h"
 
 UStretchGizmoHandleGroup::UStretchGizmoHandleGroup()
 	: Super()
diff --git a/Engine/Source/Editor/ViewportInteraction/ViewportInteractor.cpp b/Engine/Source/Editor/ViewportInteraction/ViewportInteractor.cpp
--- a/Engine/Source/Editor/ViewportInteraction/ViewportInteractor.cpp
+++ b/Engine/Source/Editor/ViewportInteraction/ViewportInteractor.cpp
@@ -101,16 +101,17 @@ bool UViewportInteractor::HandleInputKey( const FKey Key, const EInputEvent Even
 		{
 			// Selection/Movement
 			if ( ( Action->ActionType == ViewportWorldActionTypes::SelectAndMove ||
-				Action->ActionType == ViewportWorldActionTypes::SelectAndMove_LightlyPressed ) )
+				   Action->ActionType == ViewportWorldActionTypes::SelectAndMove_LightlyPressed ) )
 			{
+				const bool bIsDraggingWorld = InteractorData.DraggingMode == EViewportInteractionDraggingMode::World;
+				const bool bIsDraggingWorldWithTwoHands =
+					OtherInteractor != nullptr &&
+					( ( InteractorData.DraggingMode == EViewportInteractionDraggingMode::World && GetOtherInteractor()->GetInteractorData().DraggingMode == EViewportInteractionDraggingMode::AssistingDrag ) ||
+					  ( GetOtherInteractor()->GetInteractorData().DraggingMode == EViewportInteractionDraggingMode::World && InteractorData.DraggingMode == EViewportInteractionDraggingMode::AssistingDrag ) );
+
 				if ( Event == IE_Pressed )
 				{
 					// No clicking while we're dragging the world.  (No laser pointers are visible, anyway.)
-					const bool bIsDraggingWorldWithTwoHands =
-						OtherInteractor != nullptr &&
-						( ( InteractorData.DraggingMode == EViewportInteractionDraggingMode::World && GetOtherInteractor()->GetInteractorData().DraggingMode == EViewportInteractionDraggingMode::AssistingDrag ) ||
-						  ( GetOtherInteractor()->GetInteractorData().DraggingMode == EViewportInteractionDraggingMode::World && InteractorData.DraggingMode == EViewportInteractionDraggingMode::AssistingDrag ) );
-
 					FHitResult HitResult = GetHitResultFromLaserPointer();
 					if ( !bIsDraggingWorldWithTwoHands && !bHandled && HitResult.Actor.IsValid() )
 					{
@@ -290,8 +291,6 @@ bool UViewportInteractor::HandleInputKey( const FKey Key, const EInputEvent Even
 										// User didn't drag but did select something, so play a haptic effect.
 										PlayHapticEffect( VI::SelectionHapticFeedbackStrength->GetFloat() );
 									}
-
-									bHandled;
 								}
 							}
 						}
@@ -305,12 +304,6 @@ bool UViewportInteractor::HandleInputKey( const FKey Key, const EInputEvent Even
 						InteractorData.ClickingOnComponent = nullptr;
 					}
 
-					const bool bIsDraggingWorld = InteractorData.DraggingMode == EViewportInteractionDraggingMode::World;
-					const bool bIsDraggingWorldWithTwoHands =
-						OtherInteractor != nullptr &&
-						( ( InteractorData.DraggingMode == EViewportInteractionDraggingMode::World && GetOtherInteractor()->InteractorData.DraggingMode == EViewportInteractionDraggingMode::AssistingDrag ) ||
-						  ( GetOtherInteractor()->InteractorData.DraggingMode == EViewportInteractionDraggingMode::World && InteractorData.DraggingMode == EViewportInteractionDraggingMode::AssistingDrag ) );
-
 					// Don't allow the trigger to cancel our drag on release if we're dragging the world. 
 					if ( InteractorData.DraggingMode != EViewportInteractionDraggingMode::Nothing &&
 						!bIsDraggingWorld &&
