[INTEGRATE] Viewport Interaction: Fixed transform gizmo actor being destroyed and re-created every frame when no actor was selected (from CL 3092429)
- Happened because every frame we would call GetTransformGizmoActor() which would create it on demand, then later in the frame we'd destroy it (UpdateTranformGizmo) after seeing no objects were selected
- This hurt performance as thousands of mesh instances would be created over time (slower GC and object iteration)
- This needs to be merged to the VREditor branch. (MikeF will do that later.)

diff --git a/Engine/Source/Editor/ViewportInteraction/ViewportWorldInteraction.cpp b/Engine/Source/Editor/ViewportInteraction/ViewportWorldInteraction.cpp
--- a/Engine/Source/Editor/ViewportInteraction/ViewportWorldInteraction.cpp
+++ b/Engine/Source/Editor/ViewportInteraction/ViewportWorldInteraction.cpp
@@ -2419,6 +2419,9 @@ void UViewportWorldInteraction::RefreshTransformGizmo( const bool bNewObjectsSel
 	{
 		SpawnTransformGizmoIfNeeded();
 
+		// Make sure the gizmo is visible
+		TransformGizmoActor->GetRootComponent()->SetVisibility( true );
+
 		const ECoordSystem CurrentCoordSystem = GetTransformGizmoCoordinateSpace();
 		const bool bIsWorldSpaceGizmo = ( CurrentCoordSystem == COORD_World );
 
@@ -2591,11 +2594,10 @@ void UViewportWorldInteraction::RefreshTransformGizmo( const bool bNewObjectsSel
 	}
 	else
 	{
-		// Nothing selected, so kill our gizmo actor
-		if ( TransformGizmoActor != nullptr )
+		// Nothing selected, so hide our gizmo actor
+		if( TransformGizmoActor != nullptr )
 		{
-			DestroyTransientActor( TransformGizmoActor );
-			TransformGizmoActor = nullptr;
+			TransformGizmoActor->GetRootComponent()->SetVisibility( false );
 		}
 		GizmoLocalBounds = FBox( 0 );
 
