VR Editor: Fixed transform gizmo always being visible, even when nothing is selected (UE-35244)
- This was a regression introduced with some recent changes to mesh paint mode

diff --git a/Engine/Source/Editor/ViewportInteraction/ViewportWorldInteraction.cpp b/Engine/Source/Editor/ViewportInteraction/ViewportWorldInteraction.cpp
--- a/Engine/Source/Editor/ViewportInteraction/ViewportWorldInteraction.cpp
+++ b/Engine/Source/Editor/ViewportInteraction/ViewportWorldInteraction.cpp
@@ -2475,12 +2475,11 @@ void UViewportWorldInteraction::OnActorSelectionChanged( UObject* ChangedObject
 
 void UViewportWorldInteraction::RefreshTransformGizmo( const bool bNewObjectsSelected, bool bAllHandlesVisible, class UActorComponent* SingleVisibleHandle, const TArray< UActorComponent* >& HoveringOverHandles )
 {
-	if ( GEditor->GetSelectedActorCount() > 0 )
+	if ( GEditor->GetSelectedActorCount() > 0 && VI::ShowTransformGizmo->GetInt() != 0 && bIsTransformGizmoVisible )
 	{
 		SpawnTransformGizmoIfNeeded();
 
 		// Make sure the gizmo is visible
-		TransformGizmoActor->GetRootComponent()->SetVisibility( true );
 		TransformGizmoActor->SetIsTemporarilyHiddenInEditor( false );
 
 		const ECoordSystem CurrentCoordSystem = GetTransformGizmoCoordinateSpace();
@@ -2655,10 +2654,10 @@ void UViewportWorldInteraction::RefreshTransformGizmo( const bool bNewObjectsSel
 	}
 	else
 	{
-		// Nothing selected, so hide our gizmo actor
+		// Nothing selected or the gizmo was asked to be hidden
 		if( TransformGizmoActor != nullptr )
 		{
-			TransformGizmoActor->GetRootComponent()->SetVisibility( false );
+			TransformGizmoActor->SetIsTemporarilyHiddenInEditor( true );
 		}
 		GizmoLocalBounds = FBox( 0 );
 
@@ -2690,21 +2689,14 @@ void UViewportWorldInteraction::SpawnTransformGizmoIfNeeded()
 
 		check( TransformGizmoActor != nullptr );
 		TransformGizmoActor->SetOwnerWorldInteraction( this );
-		
-		if ( VI::ShowTransformGizmo->GetInt() == 0 || !bIsTransformGizmoVisible )
-		{
-			TransformGizmoActor->SetIsTemporarilyHiddenInEditor( true );
-		}
 	}
 }
 
 void UViewportWorldInteraction::SetTransformGizmoVisible( const bool bShouldBeVisible )
 {
 	bIsTransformGizmoVisible = bShouldBeVisible;
-	if( TransformGizmoActor != nullptr )
-	{
-		TransformGizmoActor->SetIsTemporarilyHiddenInEditor( VI::ShowTransformGizmo->GetInt() == 0 || !bIsTransformGizmoVisible );
-	}
+
+	// NOTE: The actual visibility change will be applied the next tick when RefreshTransformGizmo() is called
 }
 
 bool UViewportWorldInteraction::IsTransformGizmoVisible() const
