VR Editor: Updated all copyrights to 2016

diff --git a/Engine/Source/Editor/VREditor/Gizmo/VREditorPlaneTranslationGizmoHandle.cpp b/Engine/Source/Editor/VREditor/Gizmo/VREditorPlaneTranslationGizmoHandle.cpp
--- a/Engine/Source/Editor/VREditor/Gizmo/VREditorPlaneTranslationGizmoHandle.cpp
+++ b/Engine/Source/Editor/VREditor/Gizmo/VREditorPlaneTranslationGizmoHandle.cpp
@@ -1,4 +1,4 @@
-// Copyright 1998-2015 Epic Games, Inc. All Rights Reserved.
+// Copyright 1998-2016 Epic Games, Inc. All Rights Reserved.
 
 #include "VREditorModule.h"
 #include "VREditorPlaneTranslationGizmoHandle.h"
diff --git a/Engine/Source/Editor/VREditor/Gizmo/VREditorPlaneTranslationGizmoHandle.h b/Engine/Source/Editor/VREditor/Gizmo/VREditorPlaneTranslationGizmoHandle.h
--- a/Engine/Source/Editor/VREditor/Gizmo/VREditorPlaneTranslationGizmoHandle.h
+++ b/Engine/Source/Editor/VREditor/Gizmo/VREditorPlaneTranslationGizmoHandle.h
@@ -1,4 +1,4 @@
-// Copyright 1998-2015 Epic Games, Inc. All Rights Reserved.
+// Copyright 1998-2016 Epic Games, Inc. All Rights Reserved.
 
 #pragma once
 
diff --git a/Engine/Source/Editor/VREditor/Gizmo/VREditorRotationGizmoHandle.cpp b/Engine/Source/Editor/VREditor/Gizmo/VREditorRotationGizmoHandle.cpp
--- a/Engine/Source/Editor/VREditor/Gizmo/VREditorRotationGizmoHandle.cpp
+++ b/Engine/Source/Editor/VREditor/Gizmo/VREditorRotationGizmoHandle.cpp
@@ -1,4 +1,4 @@
-// Copyright 1998-2015 Epic Games, Inc. All Rights Reserved.
+// Copyright 1998-2016 Epic Games, Inc. All Rights Reserved.
 
 #include "VREditorModule.h"
 #include "VREditorRotationGizmoHandle.h"
diff --git a/Engine/Source/Editor/VREditor/Gizmo/VREditorRotationGizmoHandle.h b/Engine/Source/Editor/VREditor/Gizmo/VREditorRotationGizmoHandle.h
--- a/Engine/Source/Editor/VREditor/Gizmo/VREditorRotationGizmoHandle.h
+++ b/Engine/Source/Editor/VREditor/Gizmo/VREditorRotationGizmoHandle.h
@@ -1,4 +1,4 @@
-// Copyright 1998-2015 Epic Games, Inc. All Rights Reserved.
+// Copyright 1998-2016 Epic Games, Inc. All Rights Reserved.
 
 #pragma once
 
diff --git a/Engine/Source/Editor/VREditor/Gizmo/VREditorTransformGizmo.cpp b/Engine/Source/Editor/VREditor/Gizmo/VREditorTransformGizmo.cpp
--- a/Engine/Source/Editor/VREditor/Gizmo/VREditorTransformGizmo.cpp
+++ b/Engine/Source/Editor/VREditor/Gizmo/VREditorTransformGizmo.cpp
@@ -1,4 +1,4 @@
-// Copyright 1998-2015 Epic Games, Inc. All Rights Reserved.
+// Copyright 1998-2016 Epic Games, Inc. All Rights Reserved.
 
 #include "VREditorModule.h"
 #include "VREditorTransformGizmo.h"
diff --git a/Engine/Source/Editor/VREditor/Gizmo/VREditorTransformGizmo.h b/Engine/Source/Editor/VREditor/Gizmo/VREditorTransformGizmo.h
--- a/Engine/Source/Editor/VREditor/Gizmo/VREditorTransformGizmo.h
+++ b/Engine/Source/Editor/VREditor/Gizmo/VREditorTransformGizmo.h
@@ -1,4 +1,4 @@
-// Copyright 1998-2015 Epic Games, Inc. All Rights Reserved.
+// Copyright 1998-2016 Epic Games, Inc. All Rights Reserved.
 
 #pragma once
 
diff --git a/Engine/Source/Editor/VREditor/Gizmo/VREditorTranslationGizmoHandle.cpp b/Engine/Source/Editor/VREditor/Gizmo/VREditorTranslationGizmoHandle.cpp
--- a/Engine/Source/Editor/VREditor/Gizmo/VREditorTranslationGizmoHandle.cpp
+++ b/Engine/Source/Editor/VREditor/Gizmo/VREditorTranslationGizmoHandle.cpp
@@ -1,4 +1,4 @@
-// Copyright 1998-2015 Epic Games, Inc. All Rights Reserved.
+// Copyright 1998-2016 Epic Games, Inc. All Rights Reserved.
 
 #include "VREditorModule.h"
 #include "VREditorTranslationGizmoHandle.h"
diff --git a/Engine/Source/Editor/VREditor/Gizmo/VREditorTranslationGizmoHandle.h b/Engine/Source/Editor/VREditor/Gizmo/VREditorTranslationGizmoHandle.h
--- a/Engine/Source/Editor/VREditor/Gizmo/VREditorTranslationGizmoHandle.h
+++ b/Engine/Source/Editor/VREditor/Gizmo/VREditorTranslationGizmoHandle.h
@@ -1,4 +1,4 @@
-// Copyright 1998-2015 Epic Games, Inc. All Rights Reserved.
+// Copyright 1998-2016 Epic Games, Inc. All Rights Reserved.
 
 #pragma once
 
diff --git a/Engine/Source/Editor/VREditor/Interactables/VREditorButton.cpp b/Engine/Source/Editor/VREditor/Interactables/VREditorButton.cpp
--- a/Engine/Source/Editor/VREditor/Interactables/VREditorButton.cpp
+++ b/Engine/Source/Editor/VREditor/Interactables/VREditorButton.cpp
@@ -1,4 +1,4 @@
-// Copyright 1998-2015 Epic Games, Inc. All Rights Reserved.
+// Copyright 1998-2016 Epic Games, Inc. All Rights Reserved.
 
 #include "VREditorModule.h"
 #include "VREditorButton.h"
diff --git a/Engine/Source/Editor/VREditor/Interactables/VREditorButton.h b/Engine/Source/Editor/VREditor/Interactables/VREditorButton.h
--- a/Engine/Source/Editor/VREditor/Interactables/VREditorButton.h
+++ b/Engine/Source/Editor/VREditor/Interactables/VREditorButton.h
@@ -1,4 +1,4 @@
-// Copyright 1998-2015 Epic Games, Inc. All Rights Reserved.
+// Copyright 1998-2016 Epic Games, Inc. All Rights Reserved.
 
 #pragma once
 
diff --git a/Engine/Source/Editor/VREditor/Public/IVREditorModule.h b/Engine/Source/Editor/VREditor/Public/IVREditorModule.h
--- a/Engine/Source/Editor/VREditor/Public/IVREditorModule.h
+++ b/Engine/Source/Editor/VREditor/Public/IVREditorModule.h
@@ -1,4 +1,4 @@
-// Copyright 1998-2015 Epic Games, Inc. All Rights Reserved.
+// Copyright 1998-2016 Epic Games, Inc. All Rights Reserved.
 
 #pragma once
 
diff --git a/Engine/Source/Editor/VREditor/Public/VREditorInteractor.h b/Engine/Source/Editor/VREditor/Public/VREditorInteractor.h
--- a/Engine/Source/Editor/VREditor/Public/VREditorInteractor.h
+++ b/Engine/Source/Editor/VREditor/Public/VREditorInteractor.h
@@ -1,4 +1,4 @@
-// Copyright 1998-2015 Epic Games, Inc. All Rights Reserved.
+// Copyright 1998-2016 Epic Games, Inc. All Rights Reserved.
 
 #pragma once
 
diff --git a/Engine/Source/Editor/VREditor/Public/VREditorMode.h b/Engine/Source/Editor/VREditor/Public/VREditorMode.h
--- a/Engine/Source/Editor/VREditor/Public/VREditorMode.h
+++ b/Engine/Source/Editor/VREditor/Public/VREditorMode.h
@@ -1,4 +1,4 @@
-// Copyright 1998-2015 Epic Games, Inc. All Rights Reserved.
+// Copyright 1998-2016 Epic Games, Inc. All Rights Reserved.
 
 #pragma once
 
diff --git a/Engine/Source/Editor/VREditor/Teleporter/VREditorTeleporter.cpp b/Engine/Source/Editor/VREditor/Teleporter/VREditorTeleporter.cpp
--- a/Engine/Source/Editor/VREditor/Teleporter/VREditorTeleporter.cpp
+++ b/Engine/Source/Editor/VREditor/Teleporter/VREditorTeleporter.cpp
@@ -1,4 +1,4 @@
-// Copyright 1998-2015 Epic Games, Inc. All Rights Reserved.
+// Copyright 1998-2016 Epic Games, Inc. All Rights Reserved.
 
 #include "VREditorModule.h"
 #include "VREditorTeleporter.h"
diff --git a/Engine/Source/Editor/VREditor/UI/VREditorBaseUserWidget.cpp b/Engine/Source/Editor/VREditor/UI/VREditorBaseUserWidget.cpp
--- a/Engine/Source/Editor/VREditor/UI/VREditorBaseUserWidget.cpp
+++ b/Engine/Source/Editor/VREditor/UI/VREditorBaseUserWidget.cpp
@@ -1,4 +1,4 @@
-// Copyright 1998-2015 Epic Games, Inc. All Rights Reserved.
+// Copyright 1998-2016 Epic Games, Inc. All Rights Reserved.
 
 #include "VREditorModule.h"
 #include "VREditorBaseUserWidget.h"
diff --git a/Engine/Source/Editor/VREditor/UI/VREditorBaseUserWidget.h b/Engine/Source/Editor/VREditor/UI/VREditorBaseUserWidget.h
--- a/Engine/Source/Editor/VREditor/UI/VREditorBaseUserWidget.h
+++ b/Engine/Source/Editor/VREditor/UI/VREditorBaseUserWidget.h
@@ -1,4 +1,4 @@
-// Copyright 1998-2015 Epic Games, Inc. All Rights Reserved.
+// Copyright 1998-2016 Epic Games, Inc. All Rights Reserved.
 
 #pragma once
 
diff --git a/Engine/Source/Editor/VREditor/UI/VREditorDockableWindow.cpp b/Engine/Source/Editor/VREditor/UI/VREditorDockableWindow.cpp
--- a/Engine/Source/Editor/VREditor/UI/VREditorDockableWindow.cpp
+++ b/Engine/Source/Editor/VREditor/UI/VREditorDockableWindow.cpp
@@ -1,4 +1,4 @@
-// Copyright 1998-2015 Epic Games, Inc. All Rights Reserved.
+// Copyright 1998-2016 Epic Games, Inc. All Rights Reserved.
 
 #include "VREditorModule.h"
 #include "VREditorDockableWindow.h"
diff --git a/Engine/Source/Editor/VREditor/UI/VREditorDockableWindow.h b/Engine/Source/Editor/VREditor/UI/VREditorDockableWindow.h
--- a/Engine/Source/Editor/VREditor/UI/VREditorDockableWindow.h
+++ b/Engine/Source/Editor/VREditor/UI/VREditorDockableWindow.h
@@ -1,4 +1,4 @@
-// Copyright 1998-2015 Epic Games, Inc. All Rights Reserved.
+// Copyright 1998-2016 Epic Games, Inc. All Rights Reserved.
 
 #pragma once
 
diff --git a/Engine/Source/Editor/VREditor/UI/VREditorFloatingUI.cpp b/Engine/Source/Editor/VREditor/UI/VREditorFloatingUI.cpp
--- a/Engine/Source/Editor/VREditor/UI/VREditorFloatingUI.cpp
+++ b/Engine/Source/Editor/VREditor/UI/VREditorFloatingUI.cpp
@@ -1,4 +1,4 @@
-// Copyright 1998-2015 Epic Games, Inc. All Rights Reserved.
+// Copyright 1998-2016 Epic Games, Inc. All Rights Reserved.
 
 #include "VREditorModule.h"
 #include "VREditorFloatingUI.h"
diff --git a/Engine/Source/Editor/VREditor/UI/VREditorFloatingUI.h b/Engine/Source/Editor/VREditor/UI/VREditorFloatingUI.h
--- a/Engine/Source/Editor/VREditor/UI/VREditorFloatingUI.h
+++ b/Engine/Source/Editor/VREditor/UI/VREditorFloatingUI.h
@@ -1,4 +1,4 @@
-// Copyright 1998-2015 Epic Games, Inc. All Rights Reserved.
+// Copyright 1998-2016 Epic Games, Inc. All Rights Reserved.
 
 #pragma once
 
diff --git a/Engine/Source/Editor/VREditor/UI/VREditorQuickMenu.cpp b/Engine/Source/Editor/VREditor/UI/VREditorQuickMenu.cpp
--- a/Engine/Source/Editor/VREditor/UI/VREditorQuickMenu.cpp
+++ b/Engine/Source/Editor/VREditor/UI/VREditorQuickMenu.cpp
@@ -1,4 +1,4 @@
-// Copyright 1998-2015 Epic Games, Inc. All Rights Reserved.
+// Copyright 1998-2016 Epic Games, Inc. All Rights Reserved.
 
 #include "VREditorModule.h"
 #include "VREditorQuickMenu.h"
diff --git a/Engine/Source/Editor/VREditor/UI/VREditorQuickMenu.h b/Engine/Source/Editor/VREditor/UI/VREditorQuickMenu.h
--- a/Engine/Source/Editor/VREditor/UI/VREditorQuickMenu.h
+++ b/Engine/Source/Editor/VREditor/UI/VREditorQuickMenu.h
@@ -1,4 +1,4 @@
-// Copyright 1998-2015 Epic Games, Inc. All Rights Reserved.
+// Copyright 1998-2016 Epic Games, Inc. All Rights Reserved.
 
 #pragma once
 
diff --git a/Engine/Source/Editor/VREditor/UI/VREditorRadialMenu.cpp b/Engine/Source/Editor/VREditor/UI/VREditorRadialMenu.cpp
--- a/Engine/Source/Editor/VREditor/UI/VREditorRadialMenu.cpp
+++ b/Engine/Source/Editor/VREditor/UI/VREditorRadialMenu.cpp
@@ -1,4 +1,4 @@
-// Copyright 1998-2015 Epic Games, Inc. All Rights Reserved.
+// Copyright 1998-2016 Epic Games, Inc. All Rights Reserved.
 
 #include "VREditorModule.h"
 #include "VREditorRadialMenu.h"
diff --git a/Engine/Source/Editor/VREditor/UI/VREditorRadialMenu.h b/Engine/Source/Editor/VREditor/UI/VREditorRadialMenu.h
--- a/Engine/Source/Editor/VREditor/UI/VREditorRadialMenu.h
+++ b/Engine/Source/Editor/VREditor/UI/VREditorRadialMenu.h
@@ -1,4 +1,4 @@
-// Copyright 1998-2015 Epic Games, Inc. All Rights Reserved.
+// Copyright 1998-2016 Epic Games, Inc. All Rights Reserved.
 
 #pragma once
 
diff --git a/Engine/Source/Editor/VREditor/UI/VREditorRadialMenuItem.cpp b/Engine/Source/Editor/VREditor/UI/VREditorRadialMenuItem.cpp
--- a/Engine/Source/Editor/VREditor/UI/VREditorRadialMenuItem.cpp
+++ b/Engine/Source/Editor/VREditor/UI/VREditorRadialMenuItem.cpp
@@ -1,4 +1,4 @@
-// Copyright 1998-2015 Epic Games, Inc. All Rights Reserved.
+// Copyright 1998-2016 Epic Games, Inc. All Rights Reserved.
 
 #include "VREditorModule.h"
 #include "VREditorRadialMenuItem.h"
diff --git a/Engine/Source/Editor/VREditor/UI/VREditorRadialMenuItem.h b/Engine/Source/Editor/VREditor/UI/VREditorRadialMenuItem.h
--- a/Engine/Source/Editor/VREditor/UI/VREditorRadialMenuItem.h
+++ b/Engine/Source/Editor/VREditor/UI/VREditorRadialMenuItem.h
@@ -1,4 +1,4 @@
-// Copyright 1998-2015 Epic Games, Inc. All Rights Reserved.
+// Copyright 1998-2016 Epic Games, Inc. All Rights Reserved.
 
 #pragma once
 
diff --git a/Engine/Source/Editor/VREditor/UI/VREditorUISystem.cpp b/Engine/Source/Editor/VREditor/UI/VREditorUISystem.cpp
--- a/Engine/Source/Editor/VREditor/UI/VREditorUISystem.cpp
+++ b/Engine/Source/Editor/VREditor/UI/VREditorUISystem.cpp
@@ -1,4 +1,4 @@
-// Copyright 1998-2015 Epic Games, Inc. All Rights Reserved.
+// Copyright 1998-2016 Epic Games, Inc. All Rights Reserved.
 
 #include "VREditorModule.h"
 #include "VREditorUISystem.h"
diff --git a/Engine/Source/Editor/VREditor/UI/VREditorUISystem.h b/Engine/Source/Editor/VREditor/UI/VREditorUISystem.h
--- a/Engine/Source/Editor/VREditor/UI/VREditorUISystem.h
+++ b/Engine/Source/Editor/VREditor/UI/VREditorUISystem.h
@@ -1,4 +1,4 @@
-// Copyright 1998-2015 Epic Games, Inc. All Rights Reserved.
+// Copyright 1998-2016 Epic Games, Inc. All Rights Reserved.
 
 #pragma once
 
diff --git a/Engine/Source/Editor/VREditor/VREditor.Build.cs b/Engine/Source/Editor/VREditor/VREditor.Build.cs
--- a/Engine/Source/Editor/VREditor/VREditor.Build.cs
+++ b/Engine/Source/Editor/VREditor/VREditor.Build.cs
@@ -1,4 +1,4 @@
-// Copyright 1998-2015 Epic Games, Inc. All Rights Reserved.
+// Copyright 1998-2016 Epic Games, Inc. All Rights Reserved.
 
 namespace UnrealBuildTool.Rules
 {
diff --git a/Engine/Source/Editor/VREditor/VREditorAvatarActor.cpp b/Engine/Source/Editor/VREditor/VREditorAvatarActor.cpp
--- a/Engine/Source/Editor/VREditor/VREditorAvatarActor.cpp
+++ b/Engine/Source/Editor/VREditor/VREditorAvatarActor.cpp
@@ -1,4 +1,4 @@
-// Copyright 1998-2015 Epic Games, Inc. All Rights Reserved.
+// Copyright 1998-2016 Epic Games, Inc. All Rights Reserved.
 
 #include "VREditorModule.h"
 #include "VREditorAvatarActor.h"
diff --git a/Engine/Source/Editor/VREditor/VREditorAvatarActor.h b/Engine/Source/Editor/VREditor/VREditorAvatarActor.h
--- a/Engine/Source/Editor/VREditor/VREditorAvatarActor.h
+++ b/Engine/Source/Editor/VREditor/VREditorAvatarActor.h
@@ -1,4 +1,4 @@
-// Copyright 1998-2015 Epic Games, Inc. All Rights Reserved.
+// Copyright 1998-2016 Epic Games, Inc. All Rights Reserved.
 
 #pragma once
 
diff --git a/Engine/Source/Editor/VREditor/VREditorBaseActor.cpp b/Engine/Source/Editor/VREditor/VREditorBaseActor.cpp
--- a/Engine/Source/Editor/VREditor/VREditorBaseActor.cpp
+++ b/Engine/Source/Editor/VREditor/VREditorBaseActor.cpp
@@ -1,4 +1,4 @@
-// Copyright 1998-2015 Epic Games, Inc. All Rights Reserved.
+// Copyright 1998-2016 Epic Games, Inc. All Rights Reserved.
 
 #include "VREditorModule.h"
 #include "VREditorBaseActor.h"
diff --git a/Engine/Source/Editor/VREditor/VREditorBaseActor.h b/Engine/Source/Editor/VREditor/VREditorBaseActor.h
--- a/Engine/Source/Editor/VREditor/VREditorBaseActor.h
+++ b/Engine/Source/Editor/VREditor/VREditorBaseActor.h
@@ -1,4 +1,4 @@
-// Copyright 1998-2015 Epic Games, Inc. All Rights Reserved.
+// Copyright 1998-2016 Epic Games, Inc. All Rights Reserved.
 
 #pragma once
 
diff --git a/Engine/Source/Editor/VREditor/VREditorFloatingText.cpp b/Engine/Source/Editor/VREditor/VREditorFloatingText.cpp
--- a/Engine/Source/Editor/VREditor/VREditorFloatingText.cpp
+++ b/Engine/Source/Editor/VREditor/VREditorFloatingText.cpp
@@ -1,4 +1,4 @@
-// Copyright 1998-2015 Epic Games, Inc. All Rights Reserved.
+// Copyright 1998-2016 Epic Games, Inc. All Rights Reserved.
 
 #include "VREditorModule.h"
 #include "VREditorFloatingText.h"
diff --git a/Engine/Source/Editor/VREditor/VREditorFloatingText.h b/Engine/Source/Editor/VREditor/VREditorFloatingText.h
--- a/Engine/Source/Editor/VREditor/VREditorFloatingText.h
+++ b/Engine/Source/Editor/VREditor/VREditorFloatingText.h
@@ -1,4 +1,4 @@
-// Copyright 1998-2015 Epic Games, Inc. All Rights Reserved.
+// Copyright 1998-2016 Epic Games, Inc. All Rights Reserved.
 
 #pragma once
 
diff --git a/Engine/Source/Editor/VREditor/VREditorInteractor.cpp b/Engine/Source/Editor/VREditor/VREditorInteractor.cpp
--- a/Engine/Source/Editor/VREditor/VREditorInteractor.cpp
+++ b/Engine/Source/Editor/VREditor/VREditorInteractor.cpp
@@ -1,4 +1,4 @@
-// Copyright 1998-2015 Epic Games, Inc. All Rights Reserved.
+// Copyright 1998-2016 Epic Games, Inc. All Rights Reserved.
 
 #include "VREditorModule.h"
 #include "VREditorInteractor.h"
diff --git a/Engine/Source/Editor/VREditor/VREditorMode.cpp b/Engine/Source/Editor/VREditor/VREditorMode.cpp
--- a/Engine/Source/Editor/VREditor/VREditorMode.cpp
+++ b/Engine/Source/Editor/VREditor/VREditorMode.cpp
@@ -1,4 +1,4 @@
-// Copyrigh 1998-2015 Epic Games, Inc. All Rights Reserved.
+// Copyright 1998-2016 Epic Games, Inc. All Rights Reserved.
 
 #include "VREditorModule.h"
 #include "VREditorMode.h"
diff --git a/Engine/Source/Editor/VREditor/VREditorModeManager.cpp b/Engine/Source/Editor/VREditor/VREditorModeManager.cpp
--- a/Engine/Source/Editor/VREditor/VREditorModeManager.cpp
+++ b/Engine/Source/Editor/VREditor/VREditorModeManager.cpp
@@ -1,3 +1,5 @@
+// Copyright 1998-2016 Epic Games, Inc. All Rights Reserved.
+
 #include "VREditorModule.h"
 #include "VREditorModeManager.h"
 #include "VREditorMode.h"
diff --git a/Engine/Source/Editor/VREditor/VREditorModule.cpp b/Engine/Source/Editor/VREditor/VREditorModule.cpp
--- a/Engine/Source/Editor/VREditor/VREditorModule.cpp
+++ b/Engine/Source/Editor/VREditor/VREditorModule.cpp
@@ -1,4 +1,4 @@
-// Copyright 1998-2015 Epic Games, Inc. All Rights Reserved.
+// Copyright 1998-2016 Epic Games, Inc. All Rights Reserved.
 
 #include "VREditorModule.h"
 #include "IVREditorModule.h"
diff --git a/Engine/Source/Editor/VREditor/VREditorModule.h b/Engine/Source/Editor/VREditor/VREditorModule.h
--- a/Engine/Source/Editor/VREditor/VREditorModule.h
+++ b/Engine/Source/Editor/VREditor/VREditorModule.h
@@ -1,4 +1,4 @@
-// Copyright 1998-2015 Epic Games, Inc. All Rights Reserved.
+// Copyright 1998-2016 Epic Games, Inc. All Rights Reserved.
 
 #pragma once
 
diff --git a/Engine/Source/Editor/VREditor/VREditorMotionControllerInteractor.cpp b/Engine/Source/Editor/VREditor/VREditorMotionControllerInteractor.cpp
--- a/Engine/Source/Editor/VREditor/VREditorMotionControllerInteractor.cpp
+++ b/Engine/Source/Editor/VREditor/VREditorMotionControllerInteractor.cpp
@@ -1,4 +1,4 @@
-// Copyright 1998-2015 Epic Games, Inc. All Rights Reserved.
+// Copyright 1998-2016 Epic Games, Inc. All Rights Reserved.
 
 #include "VREditorModule.h"
 #include "VREditorMotionControllerInteractor.h"
diff --git a/Engine/Source/Editor/VREditor/VREditorMotionControllerInteractor.h b/Engine/Source/Editor/VREditor/VREditorMotionControllerInteractor.h
--- a/Engine/Source/Editor/VREditor/VREditorMotionControllerInteractor.h
+++ b/Engine/Source/Editor/VREditor/VREditorMotionControllerInteractor.h
@@ -1,4 +1,4 @@
-// Copyright 1998-2015 Epic Games, Inc. All Rights Reserved.
+// Copyright 1998-2016 Epic Games, Inc. All Rights Reserved.
 
 #pragma once
 
diff --git a/Engine/Source/Editor/VREditor/VREditorWidgetComponent.cpp b/Engine/Source/Editor/VREditor/VREditorWidgetComponent.cpp
--- a/Engine/Source/Editor/VREditor/VREditorWidgetComponent.cpp
+++ b/Engine/Source/Editor/VREditor/VREditorWidgetComponent.cpp
@@ -1,4 +1,4 @@
-// Copyright 1998-2015 Epic Games, Inc. All Rights Reserved.
+// Copyright 1998-2016 Epic Games, Inc. All Rights Reserved.
 
 #include "VREditorModule.h"
 #include "VREditorWidgetComponent.h"
diff --git a/Engine/Source/Editor/VREditor/VREditorWidgetComponent.h b/Engine/Source/Editor/VREditor/VREditorWidgetComponent.h
--- a/Engine/Source/Editor/VREditor/VREditorWidgetComponent.h
+++ b/Engine/Source/Editor/VREditor/VREditorWidgetComponent.h
@@ -1,4 +1,4 @@
-// Copyright 1998-2015 Epic Games, Inc. All Rights Reserved.
+// Copyright 1998-2016 Epic Games, Inc. All Rights Reserved.
 
 #pragma once
 
diff --git a/Engine/Source/Editor/VREditor/VREditorWorldInteraction.cpp b/Engine/Source/Editor/VREditor/VREditorWorldInteraction.cpp
--- a/Engine/Source/Editor/VREditor/VREditorWorldInteraction.cpp
+++ b/Engine/Source/Editor/VREditor/VREditorWorldInteraction.cpp
@@ -1,4 +1,4 @@
-// Copyright 1998-2015 Epic Games, Inc. All Rights Reserved.
+// Copyright 1998-2016 Epic Games, Inc. All Rights Reserved.
 
 #include "VREditorModule.h"
 #include "VREditorWorldInteraction.h"
diff --git a/Engine/Source/Editor/VREditor/VREditorWorldInteraction.h b/Engine/Source/Editor/VREditor/VREditorWorldInteraction.h
--- a/Engine/Source/Editor/VREditor/VREditorWorldInteraction.h
+++ b/Engine/Source/Editor/VREditor/VREditorWorldInteraction.h
@@ -1,4 +1,4 @@
-// Copyright 1998-2015 Epic Games, Inc. All Rights Reserved.
+// Copyright 1998-2016 Epic Games, Inc. All Rights Reserved.
 
 #pragma once
 
diff --git a/Engine/Source/Editor/ViewportInteraction/Gizmo/VIBaseTransformGizmo.cpp b/Engine/Source/Editor/ViewportInteraction/Gizmo/VIBaseTransformGizmo.cpp
--- a/Engine/Source/Editor/ViewportInteraction/Gizmo/VIBaseTransformGizmo.cpp
+++ b/Engine/Source/Editor/ViewportInteraction/Gizmo/VIBaseTransformGizmo.cpp
@@ -1,4 +1,4 @@
-// Copyright 1998-2015 Epic Games, Inc. All Rights Reserved.
+// Copyright 1998-2016 Epic Games, Inc. All Rights Reserved.
 
 #include "ViewportInteractionModule.h"
 #include "VIBaseTransformGizmo.h"
diff --git a/Engine/Source/Editor/ViewportInteraction/Gizmo/VIGizmoHandle.cpp b/Engine/Source/Editor/ViewportInteraction/Gizmo/VIGizmoHandle.cpp
--- a/Engine/Source/Editor/ViewportInteraction/Gizmo/VIGizmoHandle.cpp
+++ b/Engine/Source/Editor/ViewportInteraction/Gizmo/VIGizmoHandle.cpp
@@ -1,4 +1,4 @@
-// Copyright 1998-2015 Epic Games, Inc. All Rights Reserved.
+// Copyright 1998-2016 Epic Games, Inc. All Rights Reserved.
 
 #include "ViewportInteractionModule.h"
 #include "VIGizmoHandle.h"
diff --git a/Engine/Source/Editor/ViewportInteraction/Gizmo/VIPivotTransformGizmo.cpp b/Engine/Source/Editor/ViewportInteraction/Gizmo/VIPivotTransformGizmo.cpp
--- a/Engine/Source/Editor/ViewportInteraction/Gizmo/VIPivotTransformGizmo.cpp
+++ b/Engine/Source/Editor/ViewportInteraction/Gizmo/VIPivotTransformGizmo.cpp
@@ -1,4 +1,4 @@
-// Copyright 1998-2015 Epic Games, Inc. All Rights Reserved.
+// Copyright 1998-2016 Epic Games, Inc. All Rights Reserved.
 
 #include "ViewportInteractionModule.h"
 #include "VIPivotTransformGizmo.h"
diff --git a/Engine/Source/Editor/ViewportInteraction/Gizmo/VIPivotTransformGizmo.h b/Engine/Source/Editor/ViewportInteraction/Gizmo/VIPivotTransformGizmo.h
--- a/Engine/Source/Editor/ViewportInteraction/Gizmo/VIPivotTransformGizmo.h
+++ b/Engine/Source/Editor/ViewportInteraction/Gizmo/VIPivotTransformGizmo.h
@@ -1,4 +1,4 @@
-// Copyright 1998-2015 Epic Games, Inc. All Rights Reserved.
+// Copyright 1998-2016 Epic Games, Inc. All Rights Reserved.
 
 #pragma once
 
diff --git a/Engine/Source/Editor/ViewportInteraction/Gizmo/VIStretchGizmoHandle.cpp b/Engine/Source/Editor/ViewportInteraction/Gizmo/VIStretchGizmoHandle.cpp
--- a/Engine/Source/Editor/ViewportInteraction/Gizmo/VIStretchGizmoHandle.cpp
+++ b/Engine/Source/Editor/ViewportInteraction/Gizmo/VIStretchGizmoHandle.cpp
@@ -1,4 +1,4 @@
-// Copyright 1998-2015 Epic Games, Inc. All Rights Reserved.
+// Copyright 1998-2016 Epic Games, Inc. All Rights Reserved.
 
 #include "ViewportInteractionModule.h"
 #include "VIStretchGizmoHandle.h"
diff --git a/Engine/Source/Editor/ViewportInteraction/Gizmo/VIUniformScaleGizmoHandle.cpp b/Engine/Source/Editor/ViewportInteraction/Gizmo/VIUniformScaleGizmoHandle.cpp
--- a/Engine/Source/Editor/ViewportInteraction/Gizmo/VIUniformScaleGizmoHandle.cpp
+++ b/Engine/Source/Editor/ViewportInteraction/Gizmo/VIUniformScaleGizmoHandle.cpp
@@ -1,4 +1,4 @@
-// Copyright 1998-2015 Epic Games, Inc. All Rights Reserved.
+// Copyright 1998-2016 Epic Games, Inc. All Rights Reserved.
 
 #include "ViewportInteractionModule.h"
 #include "VIUniformScaleGizmoHandle.h"
diff --git a/Engine/Source/Editor/ViewportInteraction/MouseCursorInteractor.cpp b/Engine/Source/Editor/ViewportInteraction/MouseCursorInteractor.cpp
--- a/Engine/Source/Editor/ViewportInteraction/MouseCursorInteractor.cpp
+++ b/Engine/Source/Editor/ViewportInteraction/MouseCursorInteractor.cpp
@@ -1,4 +1,4 @@
-// Copyright 1998-2015 Epic Games, Inc. All Rights Reserved.
+// Copyright 1998-2016 Epic Games, Inc. All Rights Reserved.
 
 #include "ViewportInteractionModule.h"
 #include "MouseCursorInteractor.h"
diff --git a/Engine/Source/Editor/ViewportInteraction/Public/MouseCursorInteractor.h b/Engine/Source/Editor/ViewportInteraction/Public/MouseCursorInteractor.h
--- a/Engine/Source/Editor/ViewportInteraction/Public/MouseCursorInteractor.h
+++ b/Engine/Source/Editor/ViewportInteraction/Public/MouseCursorInteractor.h
@@ -1,4 +1,4 @@
-// Copyright 1998-2015 Epic Games, Inc. All Rights Reserved.
+// Copyright 1998-2016 Epic Games, Inc. All Rights Reserved.
 
 #pragma once
 
diff --git a/Engine/Source/Editor/ViewportInteraction/Public/VIBaseTransformGizmo.h b/Engine/Source/Editor/ViewportInteraction/Public/VIBaseTransformGizmo.h
--- a/Engine/Source/Editor/ViewportInteraction/Public/VIBaseTransformGizmo.h
+++ b/Engine/Source/Editor/ViewportInteraction/Public/VIBaseTransformGizmo.h
@@ -1,4 +1,4 @@
-// Copyright 1998-2015 Epic Games, Inc. All Rights Reserved.
+// Copyright 1998-2016 Epic Games, Inc. All Rights Reserved.
 
 #pragma once
 
diff --git a/Engine/Source/Editor/ViewportInteraction/Public/VIGizmoHandle.h b/Engine/Source/Editor/ViewportInteraction/Public/VIGizmoHandle.h
--- a/Engine/Source/Editor/ViewportInteraction/Public/VIGizmoHandle.h
+++ b/Engine/Source/Editor/ViewportInteraction/Public/VIGizmoHandle.h
@@ -1,4 +1,4 @@
-// Copyright 1998-2015 Epic Games, Inc. All Rights Reserved.
+// Copyright 1998-2016 Epic Games, Inc. All Rights Reserved.
 
 #pragma once
 
diff --git a/Engine/Source/Editor/ViewportInteraction/Public/VIStretchGizmoHandle.h b/Engine/Source/Editor/ViewportInteraction/Public/VIStretchGizmoHandle.h
--- a/Engine/Source/Editor/ViewportInteraction/Public/VIStretchGizmoHandle.h
+++ b/Engine/Source/Editor/ViewportInteraction/Public/VIStretchGizmoHandle.h
@@ -1,4 +1,4 @@
-// Copyright 1998-2015 Epic Games, Inc. All Rights Reserved.
+// Copyright 1998-2016 Epic Games, Inc. All Rights Reserved.
 
 #pragma once
 
diff --git a/Engine/Source/Editor/ViewportInteraction/Public/VIUniformScaleGizmoHandle.h b/Engine/Source/Editor/ViewportInteraction/Public/VIUniformScaleGizmoHandle.h
--- a/Engine/Source/Editor/ViewportInteraction/Public/VIUniformScaleGizmoHandle.h
+++ b/Engine/Source/Editor/ViewportInteraction/Public/VIUniformScaleGizmoHandle.h
@@ -1,4 +1,4 @@
-// Copyright 1998-2015 Epic Games, Inc. All Rights Reserved.
+// Copyright 1998-2016 Epic Games, Inc. All Rights Reserved.
 
 #pragma once
 
diff --git a/Engine/Source/Editor/ViewportInteraction/Public/ViewportDragOperation.h b/Engine/Source/Editor/ViewportInteraction/Public/ViewportDragOperation.h
--- a/Engine/Source/Editor/ViewportInteraction/Public/ViewportDragOperation.h
+++ b/Engine/Source/Editor/ViewportInteraction/Public/ViewportDragOperation.h
@@ -1,4 +1,4 @@
-// Copyright 1998-2015 Epic Games, Inc. All Rights Reserved.
+// Copyright 1998-2016 Epic Games, Inc. All Rights Reserved.
 
 #pragma once
 
diff --git a/Engine/Source/Editor/ViewportInteraction/Public/ViewportInteractableInterface.h b/Engine/Source/Editor/ViewportInteraction/Public/ViewportInteractableInterface.h
--- a/Engine/Source/Editor/ViewportInteraction/Public/ViewportInteractableInterface.h
+++ b/Engine/Source/Editor/ViewportInteraction/Public/ViewportInteractableInterface.h
@@ -1,4 +1,4 @@
-// Copyright 1998-2015 Epic Games, Inc. All Rights Reserved.
+// Copyright 1998-2016 Epic Games, Inc. All Rights Reserved.
 
 #pragma once
 
diff --git a/Engine/Source/Editor/ViewportInteraction/Public/ViewportInteraction.h b/Engine/Source/Editor/ViewportInteraction/Public/ViewportInteraction.h
--- a/Engine/Source/Editor/ViewportInteraction/Public/ViewportInteraction.h
+++ b/Engine/Source/Editor/ViewportInteraction/Public/ViewportInteraction.h
@@ -1,4 +1,4 @@
-// Copyright 1998-2015 Epic Games, Inc. All Rights Reserved.
+// Copyright 1998-2016 Epic Games, Inc. All Rights Reserved.
 
 #pragma once
 
diff --git a/Engine/Source/Editor/ViewportInteraction/Public/ViewportInteractionTypes.h b/Engine/Source/Editor/ViewportInteraction/Public/ViewportInteractionTypes.h
--- a/Engine/Source/Editor/ViewportInteraction/Public/ViewportInteractionTypes.h
+++ b/Engine/Source/Editor/ViewportInteraction/Public/ViewportInteractionTypes.h
@@ -1,4 +1,4 @@
-// Copyright 1998-2015 Epic Games, Inc. All Rights Reserved.
+// Copyright 1998-2016 Epic Games, Inc. All Rights Reserved.
 
 #pragma once
 
diff --git a/Engine/Source/Editor/ViewportInteraction/Public/ViewportInteractor.h b/Engine/Source/Editor/ViewportInteraction/Public/ViewportInteractor.h
--- a/Engine/Source/Editor/ViewportInteraction/Public/ViewportInteractor.h
+++ b/Engine/Source/Editor/ViewportInteraction/Public/ViewportInteractor.h
@@ -1,4 +1,4 @@
-// Copyright 1998-2015 Epic Games, Inc. All Rights Reserved.
+// Copyright 1998-2016 Epic Games, Inc. All Rights Reserved.
 
 #pragma once
 
diff --git a/Engine/Source/Editor/ViewportInteraction/Public/ViewportWorldInteraction.h b/Engine/Source/Editor/ViewportInteraction/Public/ViewportWorldInteraction.h
--- a/Engine/Source/Editor/ViewportInteraction/Public/ViewportWorldInteraction.h
+++ b/Engine/Source/Editor/ViewportInteraction/Public/ViewportWorldInteraction.h
@@ -1,4 +1,4 @@
-// Copyright 1998-2015 Epic Games, Inc. All Rights Reserved.
+// Copyright 1998-2016 Epic Games, Inc. All Rights Reserved.
 
 #pragma once
 
diff --git a/Engine/Source/Editor/ViewportInteraction/ViewportDragOperation.cpp b/Engine/Source/Editor/ViewportInteraction/ViewportDragOperation.cpp
--- a/Engine/Source/Editor/ViewportInteraction/ViewportDragOperation.cpp
+++ b/Engine/Source/Editor/ViewportInteraction/ViewportDragOperation.cpp
@@ -1,4 +1,4 @@
-// Copyright 1998-2015 Epic Games, Inc. All Rights Reserved.
+// Copyright 1998-2016 Epic Games, Inc. All Rights Reserved.
 
 #include "ViewportInteractionModule.h"
 #include "ViewportDragOperation.h"
diff --git a/Engine/Source/Editor/ViewportInteraction/ViewportInteractableInterface.cpp b/Engine/Source/Editor/ViewportInteraction/ViewportInteractableInterface.cpp
--- a/Engine/Source/Editor/ViewportInteraction/ViewportInteractableInterface.cpp
+++ b/Engine/Source/Editor/ViewportInteraction/ViewportInteractableInterface.cpp
@@ -1,4 +1,4 @@
-// Copyright 1998-2015 Epic Games, Inc. All Rights Reserved.
+// Copyright 1998-2016 Epic Games, Inc. All Rights Reserved.
 
 #include "ViewportInteractionModule.h"
 #include "ViewportInteractableInterface.h"
diff --git a/Engine/Source/Editor/ViewportInteraction/ViewportInteractionInputProcessor.cpp b/Engine/Source/Editor/ViewportInteraction/ViewportInteractionInputProcessor.cpp
--- a/Engine/Source/Editor/ViewportInteraction/ViewportInteractionInputProcessor.cpp
+++ b/Engine/Source/Editor/ViewportInteraction/ViewportInteractionInputProcessor.cpp
@@ -1,4 +1,4 @@
-// Copyright 1998-2015 Epic Games, Inc. All Rights Reserved.
+// Copyright 1998-2016 Epic Games, Inc. All Rights Reserved.
 
 #include "ViewportInteractionModule.h"
 #include "ViewportInteractionInputProcessor.h"
diff --git a/Engine/Source/Editor/ViewportInteraction/ViewportInteractionInputProcessor.h b/Engine/Source/Editor/ViewportInteraction/ViewportInteractionInputProcessor.h
--- a/Engine/Source/Editor/ViewportInteraction/ViewportInteractionInputProcessor.h
+++ b/Engine/Source/Editor/ViewportInteraction/ViewportInteractionInputProcessor.h
@@ -1,4 +1,4 @@
-// Copyright 1998-2015 Epic Games, Inc. All Rights Reserved.
+// Copyright 1998-2016 Epic Games, Inc. All Rights Reserved.
 
 #pragma once
 
diff --git a/Engine/Source/Editor/ViewportInteraction/ViewportInteractionModule.cpp b/Engine/Source/Editor/ViewportInteraction/ViewportInteractionModule.cpp
--- a/Engine/Source/Editor/ViewportInteraction/ViewportInteractionModule.cpp
+++ b/Engine/Source/Editor/ViewportInteraction/ViewportInteractionModule.cpp
@@ -1,4 +1,4 @@
-// Copyright 1998-2015 Epic Games, Inc. All Rights Reserved.
+// Copyright 1998-2016 Epic Games, Inc. All Rights Reserved.
 
 #include "ViewportInteractionModule.h"
 #include "ViewportWorldInteractionManager.h"
diff --git a/Engine/Source/Editor/ViewportInteraction/ViewportInteractionModule.h b/Engine/Source/Editor/ViewportInteraction/ViewportInteractionModule.h
--- a/Engine/Source/Editor/ViewportInteraction/ViewportInteractionModule.h
+++ b/Engine/Source/Editor/ViewportInteraction/ViewportInteractionModule.h
@@ -1,4 +1,4 @@
-// Copyright 1998-2015 Epic Games, Inc. All Rights Reserved.
+// Copyright 1998-2016 Epic Games, Inc. All Rights Reserved.
 
 #pragma once
 
diff --git a/Engine/Source/Editor/ViewportInteraction/ViewportInteractor.cpp b/Engine/Source/Editor/ViewportInteraction/ViewportInteractor.cpp
--- a/Engine/Source/Editor/ViewportInteraction/ViewportInteractor.cpp
+++ b/Engine/Source/Editor/ViewportInteraction/ViewportInteractor.cpp
@@ -1,4 +1,4 @@
-// Copyright 1998-2015 Epic Games, Inc. All Rights Reserved.
+// Copyright 1998-2016 Epic Games, Inc. All Rights Reserved.
 
 #include "ViewportInteractionModule.h"
 #include "ViewportInteractor.h"
diff --git a/Engine/Source/Editor/ViewportInteraction/ViewportInteractorData.h b/Engine/Source/Editor/ViewportInteraction/ViewportInteractorData.h
--- a/Engine/Source/Editor/ViewportInteraction/ViewportInteractorData.h
+++ b/Engine/Source/Editor/ViewportInteraction/ViewportInteractorData.h
@@ -1,4 +1,4 @@
-// Copyright 1998-2015 Epic Games, Inc. All Rights Reserved.
+// Copyright 1998-2016 Epic Games, Inc. All Rights Reserved.
 
 #pragma once
 
diff --git a/Engine/Source/Editor/ViewportInteraction/ViewportWorldInteraction.cpp b/Engine/Source/Editor/ViewportInteraction/ViewportWorldInteraction.cpp
--- a/Engine/Source/Editor/ViewportInteraction/ViewportWorldInteraction.cpp
+++ b/Engine/Source/Editor/ViewportInteraction/ViewportWorldInteraction.cpp
@@ -1,4 +1,4 @@
-// Copyright 1998-2015 Epic Games, Inc. All Rights Reserved.
+// Copyright 1998-2016 Epic Games, Inc. All Rights Reserved.
 
 #include "ViewportInteractionModule.h"
 #include "ViewportWorldInteraction.h"
diff --git a/Engine/Source/Editor/ViewportInteraction/ViewportWorldInteractionInterface.cpp b/Engine/Source/Editor/ViewportInteraction/ViewportWorldInteractionInterface.cpp
--- a/Engine/Source/Editor/ViewportInteraction/ViewportWorldInteractionInterface.cpp
+++ b/Engine/Source/Editor/ViewportInteraction/ViewportWorldInteractionInterface.cpp
@@ -1,4 +1,4 @@
-// Copyright 1998-2015 Epic Games, Inc. All Rights Reserved.
+// Copyright 1998-2016 Epic Games, Inc. All Rights Reserved.
 
 #include "ViewportInteractionModule.h"
 #include "ViewportWorldInteractionInterface.h"
diff --git a/Engine/Source/Editor/ViewportInteraction/ViewportWorldInteractionManager.cpp b/Engine/Source/Editor/ViewportInteraction/ViewportWorldInteractionManager.cpp
--- a/Engine/Source/Editor/ViewportInteraction/ViewportWorldInteractionManager.cpp
+++ b/Engine/Source/Editor/ViewportInteraction/ViewportWorldInteractionManager.cpp
@@ -1,4 +1,4 @@
-// Copyright 1998-2015 Epic Games, Inc. All Rights Reserved.
+// Copyright 1998-2016 Epic Games, Inc. All Rights Reserved.
 
 #include "ViewportInteractionModule.h"
 #include "ViewportWorldInteractionManager.h"
diff --git a/Engine/Source/Editor/ViewportInteraction/ViewportWorldInteractionManager.h b/Engine/Source/Editor/ViewportInteraction/ViewportWorldInteractionManager.h
--- a/Engine/Source/Editor/ViewportInteraction/ViewportWorldInteractionManager.h
+++ b/Engine/Source/Editor/ViewportInteraction/ViewportWorldInteractionManager.h
@@ -1,4 +1,4 @@
-// Copyright 1998-2015 Epic Games, Inc. All Rights Reserved.
+// Copyright 1998-2016 Epic Games, Inc. All Rights Reserved.
 
 #pragma once
 
