VR Editor: Allow objects to snap to ground while translating them (press touchpad or analog stick when dragging object) #jira UETOOL-670

diff --git a/Engine/Source/Editor/VREditor/VREditorMotionControllerInteractor.cpp b/Engine/Source/Editor/VREditor/VREditorMotionControllerInteractor.cpp
--- a/Engine/Source/Editor/VREditor/VREditorMotionControllerInteractor.cpp
+++ b/Engine/Source/Editor/VREditor/VREditorMotionControllerInteractor.cpp
@@ -618,6 +618,32 @@ void UVREditorMotionControllerInteractor::HandleInputKey( FViewportActionKeyInpu
 		}
 	}
 
+	if( Action.ActionType == VRActionTypes::ConfirmRadialSelection )
+	{
+		bOutWasHandled = true;
+		const EViewportInteractionDraggingMode DraggingMode = GetDraggingMode();
+
+		if( Event == IE_Pressed )
+		{ 
+			// Start dragging at laser impact when already dragging actors freely
+			if( DraggingMode == EViewportInteractionDraggingMode::ActorsFreely )
+			{
+				FVector PlaceAt = GetHoverLocation();
+				const bool bIsPlacingActors = true;
+				const FViewportActionKeyInput FakeSelectAndMoveAction( ViewportWorldActionTypes::SelectAndMove_LightlyPressed );
+				WorldInteraction->StartDraggingActors( this, FakeSelectAndMoveAction, WorldInteraction->GetTransformGizmoActor()->GetRootComponent(), PlaceAt, bIsPlacingActors );
+			}
+		}
+		else if( Event == IE_Released )
+		{
+			// Disable dragging at laser impact when releasing
+			if( DraggingMode == EViewportInteractionDraggingMode::ActorsAtLaserImpact )
+			{
+				SetDraggingMode( EViewportInteractionDraggingMode::ActorsFreely );
+			}
+		}
+	}
+
 	ApplyButtonPressColors( Action );
 }
 
diff --git a/Engine/Source/Editor/ViewportInteraction/Public/ViewportWorldInteraction.h b/Engine/Source/Editor/ViewportInteraction/Public/ViewportWorldInteraction.h
--- a/Engine/Source/Editor/ViewportInteraction/Public/ViewportWorldInteraction.h
+++ b/Engine/Source/Editor/ViewportInteraction/Public/ViewportWorldInteraction.h
@@ -364,6 +364,9 @@ private:
 	/** Spawns the grid actor */
 	void SpawnGridMeshActor();
 
+	/** Average location of all the current transformables */
+	FVector CalculateAverageLocationOfTransformables();
+
 	//
 	// Colors
 	//
diff --git a/Engine/Source/Editor/ViewportInteraction/ViewportInteractorData.h b/Engine/Source/Editor/ViewportInteraction/ViewportInteractorData.h
--- a/Engine/Source/Editor/ViewportInteraction/ViewportInteractorData.h
+++ b/Engine/Source/Editor/ViewportInteraction/ViewportInteractorData.h
@@ -136,6 +136,10 @@ struct FViewportInteractorData
 	/** If the latest hitresult is hovering over a priority type */
 	bool bHitResultIsPriorityType;
 
+	/** The offset between the hitlocation of the object selected to start dragging and its center. 
+		This is used to offset the objects when dragging to the end of the laser */
+	FVector StartHitLocationToTransformableCenter;
+
 	/** Default constructor for FVirtualHand that initializes safe defaults */
 	FViewportInteractorData()
 	{
@@ -171,6 +175,8 @@ struct FViewportInteractorData
 		bAllowTriggerLightPressLocking = true;
 		bAllowTriggerFullPress = true;
 		bHitResultIsPriorityType = false;
+
+		StartHitLocationToTransformableCenter = FVector::ZeroVector;
 	}
 };
 
diff --git a/Engine/Source/Editor/ViewportInteraction/ViewportWorldInteraction.cpp b/Engine/Source/Editor/ViewportInteraction/ViewportWorldInteraction.cpp
--- a/Engine/Source/Editor/ViewportInteraction/ViewportWorldInteraction.cpp
+++ b/Engine/Source/Editor/ViewportInteraction/ViewportWorldInteraction.cpp
@@ -2148,8 +2148,7 @@ void UViewportWorldInteraction::StartDraggingActors( UViewportInteractor* Intera
 		const bool bHaveLaserPointer = Interactor->GetLaserPointer( /* Out */ LaserPointerStart, /* Out */ LaserPointerEnd );
 		if ( bHaveLaserPointer )
 		{
-			AActor* Actor = ClickedComponent->GetOwner();
-
+			AActor* Actor = ClickedComponent->GetOwner(); 
 			FViewportInteractorData& InteractorData = Interactor->GetInteractorData();
 
 			// Capture undo state
@@ -2227,6 +2226,9 @@ void UViewportWorldInteraction::StartDraggingActors( UViewportInteractor* Intera
 
 			SetupTransformablesForSelectedActors();
 
+			// Calculate the offset between the average location and the hit location after setting the transformables for the new selected objects
+			InteractorData.StartHitLocationToTransformableCenter = CalculateAverageLocationOfTransformables() - HitLocation;
+
 			// If we're placing actors, start interpolating to their actual location.  This helps smooth everything out when
 			// using the laser impact point as the target transform
 			if ( bIsPlacingActors )
@@ -2396,7 +2398,9 @@ bool UViewportWorldInteraction::FindPlacementPointUnderLaser( UViewportInteracto
 			}
 
 			const FVector WorldSpaceExtremePointOnBox = InteractorData.GizmoStartTransform.TransformVectorNoScale( ExtremePointOnBox );
+
 			HitLocation -= WorldSpaceExtremePointOnBox;
+			HitLocation -= InteractorData.StartHitLocationToTransformableCenter;
 		}
 
 		OutHitLocation = HitLocation;
@@ -2941,6 +2945,20 @@ void UViewportWorldInteraction::SpawnGridMeshActor()
 	}
 }
 
+FVector UViewportWorldInteraction::CalculateAverageLocationOfTransformables()
+{
+	FVector Result = FVector::ZeroVector;
+
+	for( TUniquePtr<FViewportTransformable>& TransformablePtr : Transformables )
+	{
+		Result += ((*TransformablePtr).LastTransform.GetLocation());
+	}
+
+	Result /= Transformables.Num();
+
+	return Result;
+}
+
 FLinearColor UViewportWorldInteraction::GetColor( const EColors Color ) const
 {
 	return Colors[ (int32)Color ];
