ViewportInteraction refactoring; miscellaneous VREditor improvements

- This merge contains various changes to ViewportInteraction and VR Editor to make things more extensible
       - Some changes are from Yannick's recent shelved refactoring work, that I've modified and improved upon
       - EditorWorldManager is gone and replaced by EditorWorldExtensionManager
       - EditorWorldExtensions are sort of like a more modern version of 'FEdMode', but still a work in progress
       - Cleaned up access to VREditorMode from other modules, forcing systems to go through EditorWorldExtensions

- Overhauled how transforming objects works with world interaction
       - Viewport interactors can now be used to move objects other than actors around (by implementing an UViewportTransformer, and a FViewportTransformable)
       - Undo/Redo now works better with inertial transformation!  The transaction only ends when objects finally come to rest.
       - Some initial support for 'grabber sphere' interactor methods has been implemented (not used yet)
       - Viewport interaction input events now receive the viewport being interacted through
       - Viewport interaction hover events no longer get a viewport client (because they must be designed to work any number of viewports.)

NOTE: This introduces UBT warnings about cyclic module dependencies.  We'll have to address this in a different changelist.

Other changes:
- The active Viewport Interaction 'gizmo mode' is now tied directly to the editor's normal gizmo mode.  They share the same state.
- New console variable 'VI.UseTransientActors' can be turned off to force editor actors to be created non-transient to make it easier to debug
- New console variable 'VI.DragTranslationVelocityStopEpsilon' sets the speed at which transformables will stop have inertia applied to them
- Fixed cyclic dependencies with ViewportInteraction and VREditor modules not being tagged properly for UBT
- Fixed some issues with transform gizmos not getting release events
- Various methods that should have been const were made const
- Eliminated duplicate implementation of SpawnTransientSceneActor and DestroyTransientActor; made it static
- UnrealEd no longer directly depends on VREditor and ViewportInteraction modules
- Engine: AActor::SetIsTemporarilyHiddenInEditor no longer does any work if the object's hidden state didn't change
- Slate: Input preprocessor (if bound) now also gets a chance at mouse button down and up events

diff --git a/Engine/Plugins/MovieScene/LevelSequenceEditor/Source/LevelSequenceEditor/Private/LevelSequenceEditorToolkit.cpp b/Engine/Plugins/MovieScene/LevelSequenceEditor/Source/LevelSequenceEditor/Private/LevelSequenceEditorToolkit.cpp
--- a/Engine/Plugins/MovieScene/LevelSequenceEditor/Source/LevelSequenceEditor/Private/LevelSequenceEditorToolkit.cpp
+++ b/Engine/Plugins/MovieScene/LevelSequenceEditor/Source/LevelSequenceEditor/Private/LevelSequenceEditorToolkit.cpp
@@ -53,9 +53,7 @@
 #include "Tracks/MovieScene3DTransformTrack.h"
 
 // To override Sequencer editor behavior for VR Editor 
-#include "IVREditorModule.h"
 #include "VREditorMode.h"
-#include "EditorWorldManager.h"
 
 
 #define LOCTEXT_NAMESPACE "LevelSequenceEditor"
@@ -279,20 +277,20 @@ void FLevelSequenceEditorToolkit::Initialize(const EToolkitMode::Type Mode, cons
 	LevelSequenceEditorModule.OnMasterSequenceCreated().AddRaw(this, &FLevelSequenceEditorToolkit::HandleMasterSequenceCreated);
 
 	FLevelSequenceEditorToolkit::OnOpened().Broadcast(*this);
-	
-	// If we are currently in the VR editor
-	bool bIsInVREditor = IVREditorModule::Get().IsVREditorModeActive();
-	if (bIsInVREditor)
-	{
-		const UWorld* World = Cast<UWorld>(GetLevelSequenceEditorPlaybackContext());
-		UVREditorMode* VRMode = GEditor->GetEditorWorldManager()->GetEditorWorldWrapper(World)->GetVREditorMode();
-		VRMode->OnVREditingModeExit_Handler.BindSP(this, &FLevelSequenceEditorToolkit::HandleVREditorModeExit);
-		VRMode->SaveSequencerSettings(Sequencer->GetKeyAllEnabled(), Sequencer->GetAutoKeyMode());
-		// Override currently set autokey behavior to always autokey all
-		Sequencer->SetAutoKeyMode(EAutoKeyMode::KeyAll);
-		Sequencer->SetKeyAllEnabled(true);
-		// Tell the VR Editor mode that Sequencer has refreshed
-		VRMode->RefreshVREditorSequencer();
+
+	{
+		const UWorld* World = CastChecked<UWorld>( GetLevelSequenceEditorPlaybackContext() );
+		UVREditorMode* VRMode = CastChecked<UVREditorMode>( GEditor->GetEditorWorldExtensionsManager()->GetEditorWorldExtensions( World )->FindExtension( UVREditorMode::StaticClass() ) );
+		if (VRMode != nullptr)
+		{
+			VRMode->OnVREditingModeExit_Handler.BindSP(this, &FLevelSequenceEditorToolkit::HandleVREditorModeExit);
+			VRMode->SaveSequencerSettings(Sequencer->GetKeyAllEnabled(), Sequencer->GetAutoKeyMode());
+			// Override currently set autokey behavior to always autokey all
+			Sequencer->SetAutoKeyMode(EAutoKeyMode::KeyAll);
+			Sequencer->SetKeyAllEnabled(true);
+			// Tell the VR Editor mode that Sequencer has refreshed
+			VRMode->RefreshVREditorSequencer();
+		}
 	}
 }
 
@@ -728,7 +726,7 @@ void FLevelSequenceEditorToolkit::HandleActorAddedToSequencer(AActor* Actor, con
 void FLevelSequenceEditorToolkit::HandleVREditorModeExit()
 {
 	const UWorld* World = Cast<UWorld>(GetLevelSequenceEditorPlaybackContext());
-	UVREditorMode* VRMode = GEditor->GetEditorWorldManager()->GetEditorWorldWrapper(World)->GetVREditorMode();
+	UVREditorMode* VRMode = CastChecked<UVREditorMode>( GEditor->GetEditorWorldExtensionsManager()->GetEditorWorldExtensions( World )->FindExtension( UVREditorMode::StaticClass() ) );
 
 	// Reset auto key settings
 	Sequencer->SetAutoKeyMode(VRMode->GetSavedEditorState().AutoKeyMode);
diff --git a/Engine/Source/Editor/FoliageEdit/FoliageEdit.build.cs b/Engine/Source/Editor/FoliageEdit/FoliageEdit.build.cs
--- a/Engine/Source/Editor/FoliageEdit/FoliageEdit.build.cs
+++ b/Engine/Source/Editor/FoliageEdit/FoliageEdit.build.cs
@@ -6,7 +6,7 @@ public class FoliageEdit : ModuleRules
 {
 	public FoliageEdit(TargetInfo Target)
 	{
-		PublicDependencyModuleNames.AddRange(
+		PrivateDependencyModuleNames.AddRange(
 			new string[] 
 			{
 				"Core",
@@ -24,15 +24,18 @@ public class FoliageEdit : ModuleRules
                 "DetailCustomizations",
                 "AssetTools",
                 "Foliage",
+				"ViewportInteraction",
+				"VREditor"
 			}
 		);
 
-		PrivateDependencyModuleNames.AddRange(
+		CircularlyReferencedDependentModules.AddRange(
 			new string[]
 			{
 				"ViewportInteraction",
 				"VREditor"
 			}
 		);
-    }
+
+	}
 }
diff --git a/Engine/Source/Editor/FoliageEdit/Private/FoliageEdMode.cpp b/Engine/Source/Editor/FoliageEdit/Private/FoliageEdMode.cpp
--- a/Engine/Source/Editor/FoliageEdit/Private/FoliageEdMode.cpp
+++ b/Engine/Source/Editor/FoliageEdit/Private/FoliageEdMode.cpp
@@ -42,11 +42,10 @@
 #include "Components/BrushComponent.h"
 
 // VR Editor
-#include "IVREditorModule.h"
 #include "VREditorMode.h"
 #include "ViewportWorldInteraction.h"
 #include "VREditorInteractor.h"
-#include "EditorWorldManager.h"
+#include "EditorWorldExtension.h"
 
 
 #define LOCTEXT_NAMESPACE "FoliageEdMode"
@@ -357,7 +356,7 @@ void FEdModeFoliage::Enter()
 	NotifyNewCurrentLevel();
 
 	// Register for VR input events
-	UViewportWorldInteraction* ViewportWorldInteraction = GEditor->GetEditorWorldManager()->GetEditorWorldWrapper(GetWorld())->GetViewportWorldInteraction();
+	UViewportWorldInteraction* ViewportWorldInteraction = Cast<UViewportWorldInteraction>( GEditor->GetEditorWorldExtensionsManager()->GetEditorWorldExtensions( GetWorld() )->FindExtension( UViewportWorldInteraction::StaticClass() ) );
 	if (ViewportWorldInteraction != nullptr)
 	{
 		ViewportWorldInteraction->OnViewportInteractionInputAction().RemoveAll(this);
@@ -372,13 +371,12 @@ void FEdModeFoliage::Enter()
 void FEdModeFoliage::Exit()
 {
 	// Unregister VR mode from event handlers
-	if (IVREditorModule::IsAvailable())
 	{
-		UViewportWorldInteraction* WorldInteraction = GEditor->GetEditorWorldManager()->GetEditorWorldWrapper(GetWorld())->GetViewportWorldInteraction();
-		if (WorldInteraction != nullptr)
+		UViewportWorldInteraction* ViewportWorldInteraction = Cast<UViewportWorldInteraction>( GEditor->GetEditorWorldExtensionsManager()->GetEditorWorldExtensions( GetWorld() )->FindExtension( UViewportWorldInteraction::StaticClass() ) );
+		if (ViewportWorldInteraction != nullptr)
 		{
-			WorldInteraction->OnViewportInteractionInputAction().RemoveAll(this);
-			WorldInteraction->OnViewportInteractionHoverUpdate().RemoveAll(this);
+			ViewportWorldInteraction->OnViewportInteractionInputAction().RemoveAll(this);
+			ViewportWorldInteraction->OnViewportInteractionHoverUpdate().RemoveAll(this);
 			FoliageInteractor = nullptr;
 		}
 	}
@@ -439,13 +437,11 @@ void FEdModeFoliage::Exit()
 	FEdMode::Exit();
 }
 
-void FEdModeFoliage::OnVRHoverUpdate(FEditorViewportClient& ViewportClient, UViewportInteractor* Interactor, FVector& HoverImpactPoint, bool& bWasHandled)
+void FEdModeFoliage::OnVRHoverUpdate(UViewportInteractor* Interactor, FVector& HoverImpactPoint, bool& bWasHandled)
 {
-	// Check if VR Editor is active
-	if (IVREditorModule::IsAvailable())
 	{
-		UVREditorMode* VREditorMode = GEditor->GetEditorWorldManager()->GetEditorWorldWrapper(ViewportClient.GetWorld())->GetVREditorMode();
-		if (VREditorMode != nullptr && VREditorMode->IsFullyInitialized())
+		UVREditorMode* VREditorMode = Cast<UVREditorMode>( GEditor->GetEditorWorldExtensionsManager()->GetEditorWorldExtensions( GetWorld() )->FindExtension( UVREditorMode::StaticClass() ) );
+		if( VREditorMode != nullptr && VREditorMode->IsFullyInitialized() )
 		{
 			// Check if we're hovering over UI. If so, stop painting so we don't display the preview brush sphere
 			if (FoliageInteractor && FoliageInteractor->IsHoveringOverPriorityType())
@@ -473,7 +469,7 @@ void FEdModeFoliage::OnVRHoverUpdate(FEditorViewportClient& ViewportClient, UVie
 				{
 					const FVector LaserPointerDirection = (LaserPointerEnd - LaserPointerStart).GetSafeNormal();
 
-					FoliageBrushTrace(&ViewportClient, LaserPointerStart, LaserPointerDirection);
+					FoliageBrushTrace(nullptr, LaserPointerStart, LaserPointerDirection);
 				}
 			}
 			
@@ -483,7 +479,7 @@ void FEdModeFoliage::OnVRHoverUpdate(FEditorViewportClient& ViewportClient, UVie
 
 void FEdModeFoliage::OnVRAction(class FEditorViewportClient& ViewportClient, UViewportInteractor* Interactor, const FViewportActionKeyInput& Action, bool& bOutIsInputCaptured, bool& bWasHandled)
 {
-	UVREditorMode* VREditorMode = GEditor->GetEditorWorldManager()->GetEditorWorldWrapper(ViewportClient.GetWorld())->GetVREditorMode();
+	UVREditorMode* VREditorMode = Cast<UVREditorMode>( GEditor->GetEditorWorldExtensionsManager()->GetEditorWorldExtensions( GetWorld() )->FindExtension( UVREditorMode::StaticClass() ) );
 	if (VREditorMode != nullptr && Interactor != nullptr)
 	{
 		const UVREditorInteractor* VRInteractor = Cast<UVREditorInteractor>(Interactor);
@@ -784,7 +780,7 @@ void FEdModeFoliage::EndFoliageBrushTrace()
 void FEdModeFoliage::FoliageBrushTrace(FEditorViewportClient* ViewportClient, const FVector& InRayOrigin, const FVector& InRayDirection)
 {
 	bBrushTraceValid = false;
-	if (!ViewportClient->IsMovingCamera() && ViewportClient->IsVisible())
+	if (ViewportClient == nullptr || ( !ViewportClient->IsMovingCamera() && ViewportClient->IsVisible() ) )
 	{
 		if (UISettings.GetPaintToolSelected() || UISettings.GetReapplyToolSelected() || UISettings.GetLassoSelectToolSelected())
 		{
@@ -792,7 +788,7 @@ void FEdModeFoliage::FoliageBrushTrace(FEditorViewportClient* ViewportClient, co
 			const FVector TraceEnd(InRayOrigin + InRayDirection * HALF_WORLD_MAX);
 
 			FHitResult Hit;
-			UWorld* World = ViewportClient->GetWorld();
+			UWorld* World = GetWorld();
 			static FName NAME_FoliageBrush = FName(TEXT("FoliageBrush"));
 			FFoliagePaintingGeometryFilter FilterFunc = FFoliagePaintingGeometryFilter(UISettings);
 
@@ -829,8 +825,8 @@ void FEdModeFoliage::FoliageBrushTrace(FEditorViewportClient* ViewportClient, co
 bool FEdModeFoliage::MouseMove(FEditorViewportClient* ViewportClient, FViewport* Viewport, int32 MouseX, int32 MouseY)
 {
 	// Use mouse capture if there's no other interactor currently tracing brush
-	UVREditorMode* VREditorMode = GEditor->GetEditorWorldManager()->GetEditorWorldWrapper(ViewportClient->GetWorld())->GetVREditorMode();
-	if (!VREditorMode->IsActive())
+	UVREditorMode* VREditorMode = Cast<UVREditorMode>( GEditor->GetEditorWorldExtensionsManager()->GetEditorWorldExtensions( GetWorld() )->FindExtension( UVREditorMode::StaticClass() ) );
+	if (VREditorMode == nullptr || !VREditorMode->IsActive())
 	{
 		// Compute a world space ray from the screen space mouse coordinates
 		FSceneViewFamilyContext ViewFamily(FSceneViewFamily::ConstructionValues(
@@ -867,8 +863,8 @@ bool FEdModeFoliage::MouseMove(FEditorViewportClient* ViewportClient, FViewport*
 bool FEdModeFoliage::CapturedMouseMove(FEditorViewportClient* ViewportClient, FViewport* Viewport, int32 MouseX, int32 MouseY)
 {
 	// Use mouse capture if there's no other interactor currently tracing brush
-	UVREditorMode* VREditorMode = GEditor->GetEditorWorldManager()->GetEditorWorldWrapper(ViewportClient->GetWorld())->GetVREditorMode();
-	if (!VREditorMode->IsActive())
+	UVREditorMode* VREditorMode = Cast<UVREditorMode>( GEditor->GetEditorWorldExtensionsManager()->GetEditorWorldExtensions( GetWorld() )->FindExtension( UVREditorMode::StaticClass() ) );
+	if (VREditorMode == nullptr || !VREditorMode->IsActive())
 	{
 		//Compute a world space ray from the screen space mouse coordinates
 		FSceneViewFamilyContext ViewFamily(FSceneViewFamily::ConstructionValues(
@@ -3461,4 +3457,4 @@ void FFoliageUISettings::Save()
 	GConfig->SetFloat(TEXT("FoliageEdit"), TEXT("PaletteThumbnailScale"), PaletteThumbnailScale, GEditorPerProjectIni);
 }
 
-#undef LOCTEXT_NAMESPACE
+#undef LOCTEXT_NAMESPACE
\ No newline at end of file
diff --git a/Engine/Source/Editor/FoliageEdit/Private/FoliageEdMode.h b/Engine/Source/Editor/FoliageEdit/Private/FoliageEdMode.h
--- a/Engine/Source/Editor/FoliageEdit/Private/FoliageEdMode.h
+++ b/Engine/Source/Editor/FoliageEdit/Private/FoliageEdMode.h
@@ -456,7 +456,7 @@ public:
 	void OnVRAction(class FEditorViewportClient& ViewportClient, class UViewportInteractor* Interactor, const struct FViewportActionKeyInput& Action, bool& bOutIsInputCaptured, bool& bWasHandled);
 
 	/** Called on VR hovering */
-	void OnVRHoverUpdate(FEditorViewportClient& ViewportClient, UViewportInteractor* Interactor, FVector& HoverImpactPoint, bool& bWasHandled);
+	void OnVRHoverUpdate(UViewportInteractor* Interactor, FVector& HoverImpactPoint, bool& bWasHandled);
 
 	typedef TMap<FName, TMap<ULandscapeComponent*, TArray<uint8> > > LandscapeLayerCacheData;
 
diff --git a/Engine/Source/Editor/LandscapeEditor/LandscapeEditor.Build.cs b/Engine/Source/Editor/LandscapeEditor/LandscapeEditor.Build.cs
--- a/Engine/Source/Editor/LandscapeEditor/LandscapeEditor.Build.cs
+++ b/Engine/Source/Editor/LandscapeEditor/LandscapeEditor.Build.cs
@@ -28,6 +28,14 @@ public class LandscapeEditor : ModuleRules
 			}
 			);
 
+		CircularlyReferencedDependentModules.AddRange(
+			new string[]
+			{
+				"ViewportInteraction",
+				"VREditor"
+			}
+		);
+
 		PrivateIncludePathModuleNames.AddRange(
 			new string[] {
 				"MainFrame",
diff --git a/Engine/Source/Editor/LandscapeEditor/Private/LandscapeEdMode.cpp b/Engine/Source/Editor/LandscapeEditor/Private/LandscapeEdMode.cpp
--- a/Engine/Source/Editor/LandscapeEditor/Private/LandscapeEdMode.cpp
+++ b/Engine/Source/Editor/LandscapeEditor/Private/LandscapeEdMode.cpp
@@ -26,7 +26,6 @@
 #include "Toolkits/ToolkitManager.h"
 #include "LandscapeHeightfieldCollisionComponent.h"
 #include "InstancedFoliageActor.h"
-#include "IVREditorModule.h"
 #include "ViewportWorldInteraction.h"
 #include "VREditorInteractor.h"
 #include "LandscapeEdModeTools.h"
@@ -39,7 +38,7 @@
 
 // VR Editor
 #include "VREditorMode.h"
-#include "EditorWorldManager.h"
+#include "EditorWorldExtension.h"
 
 // Classes
 #include "LandscapeMaterialInstanceConstant.h"
@@ -475,7 +474,7 @@ void FEdModeLandscape::Enter()
 	}
 
 	// Register to find out about VR input events
-	UViewportWorldInteraction* ViewportWorldInteraction = GEditor->GetEditorWorldManager()->GetEditorWorldWrapper(GetWorld())->GetViewportWorldInteraction();
+	UViewportWorldInteraction* ViewportWorldInteraction = Cast<UViewportWorldInteraction>(GEditor->GetEditorWorldExtensionsManager()->GetEditorWorldExtensions(GetWorld())->FindExtension(UViewportWorldInteraction::StaticClass()));
 	if (ViewportWorldInteraction != nullptr)
 	{
 		ViewportWorldInteraction->OnViewportInteractionInputAction().RemoveAll(this);
@@ -491,14 +490,11 @@ void FEdModeLandscape::Enter()
 void FEdModeLandscape::Exit()
 {
 	// Unregister VR mode from event handlers
-	if (IVREditorModule::IsAvailable())
+	UViewportWorldInteraction* ViewportWorldInteraction = Cast<UViewportWorldInteraction>(GEditor->GetEditorWorldExtensionsManager()->GetEditorWorldExtensions(GetWorld())->FindExtension(UViewportWorldInteraction::StaticClass()));
+	if (ViewportWorldInteraction != nullptr)
 	{
-		UViewportWorldInteraction* ViewportWorldInteraction = GEditor->GetEditorWorldManager()->GetEditorWorldWrapper(GetWorld())->GetViewportWorldInteraction();
-		if (ViewportWorldInteraction != nullptr)
-		{
-			ViewportWorldInteraction->OnViewportInteractionInputAction().RemoveAll(this);
-			ViewportWorldInteraction->OnViewportInteractionHoverUpdate().RemoveAll(this);
-		}
+		ViewportWorldInteraction->OnViewportInteractionInputAction().RemoveAll(this);
+		ViewportWorldInteraction->OnViewportInteractionHoverUpdate().RemoveAll(this);
 	}
 	
 	FEditorSupportDelegates::WorldChange.Remove(OnWorldChangeDelegateHandle);
@@ -570,11 +566,11 @@ void FEdModeLandscape::Exit()
 }
 
 
-void FEdModeLandscape::OnVRHoverUpdate(FEditorViewportClient& ViewportClient, UViewportInteractor* Interactor, FVector& HoverImpactPoint, bool& bWasHandled)
+void FEdModeLandscape::OnVRHoverUpdate(UViewportInteractor* Interactor, FVector& HoverImpactPoint, bool& bWasHandled)
 {
 	if( InteractorPainting != nullptr && InteractorPainting == Interactor )
 	{
-		UVREditorMode* VREditorMode = GEditor->GetEditorWorldManager()->GetEditorWorldWrapper( ViewportClient.GetWorld() )->GetVREditorMode();
+		UVREditorMode* VREditorMode = Cast<UVREditorMode>( GEditor->GetEditorWorldExtensionsManager()->GetEditorWorldExtensions( GetWorld() )->FindExtension( UVREditorMode::StaticClass() ) );
 		if( VREditorMode != nullptr && VREditorMode->IsActive() && Interactor != nullptr )
 		{
 			const UVREditorInteractor* VRInteractor = Cast<UVREditorInteractor>( Interactor );
@@ -585,16 +581,8 @@ void FEdModeLandscape::OnVRHoverUpdate(FEditorViewportClient& ViewportClient, UV
 				FVector LaserPointerStart, LaserPointerEnd;
 				if( Interactor->GetLaserPointer( /* Out */ LaserPointerStart, /* Out */ LaserPointerEnd ) )
 				{
-					if( LandscapeTrace( &ViewportClient, LaserPointerStart, LaserPointerEnd, HitLocation ) )
+					if( LandscapeTrace( LaserPointerStart, LaserPointerEnd, HitLocation ) )
 					{
-						if( bIsPaintingInVR && !( CurrentToolTarget.TargetType == ELandscapeToolTargetType::Weightmap && CurrentToolTarget.LayerInfo == NULL ) )
-						{
-							if( CurrentTool->BeginTool( &ViewportClient, CurrentToolTarget, HitLocation, Interactor ) )
-							{
-								ToolActiveViewport = ViewportClient.Viewport;
-							}
-						}
-
 						if( CurrentBrush )
 						{
 							// Inform the brush of the current location, to update the cursor
@@ -609,7 +597,7 @@ void FEdModeLandscape::OnVRHoverUpdate(FEditorViewportClient& ViewportClient, UV
 
 void FEdModeLandscape::OnVRAction(FEditorViewportClient& ViewportClient, UViewportInteractor* Interactor, const struct FViewportActionKeyInput& Action, bool& bOutIsInputCaptured, bool& bWasHandled)
 {
-	UVREditorMode* VREditorMode = GEditor->GetEditorWorldManager()->GetEditorWorldWrapper(ViewportClient.GetWorld())->GetVREditorMode();
+	UVREditorMode* VREditorMode = Cast<UVREditorMode>( GEditor->GetEditorWorldExtensionsManager()->GetEditorWorldExtensions( GetWorld() )->FindExtension( UVREditorMode::StaticClass() ) );
 	// Never show the traditional Unreal transform widget.  It doesn't work in VR because we don't have hit proxies.
 	ViewportClient.EngineShowFlags.SetModeWidgets(false);
 
@@ -634,7 +622,7 @@ void FEdModeLandscape::OnVRAction(FEditorViewportClient& ViewportClient, UViewpo
 					FVector LaserPointerStart, LaserPointerEnd;
 					if (Interactor->GetLaserPointer( /* Out */ LaserPointerStart, /* Out */ LaserPointerEnd))
 					{
-						if (LandscapeTrace(&ViewportClient, LaserPointerStart, LaserPointerEnd, HitLocation))
+						if (LandscapeTrace(LaserPointerStart, LaserPointerEnd, HitLocation))
 						{
 							if (!(CurrentToolTarget.TargetType == ELandscapeToolTargetType::Weightmap && CurrentToolTarget.LayerInfo == NULL))
 							{
@@ -895,16 +883,16 @@ bool FEdModeLandscape::LandscapeMouseTrace(FEditorViewportClient* ViewportClient
 
 	FVector End = Start + WORLD_MAX * BrushTraceDirection;
 
-	return LandscapeTrace(ViewportClient, Start, End, OutHitLocation);
+	return LandscapeTrace(Start, End, OutHitLocation);
 }
 
-bool FEdModeLandscape::LandscapeTrace(const FEditorViewportClient* ViewportClient, const FVector& InRayOrigin, const FVector& InRayEnd, FVector& OutHitLocation)
+bool FEdModeLandscape::LandscapeTrace(const FVector& InRayOrigin, const FVector& InRayEnd, FVector& OutHitLocation)
 {
 	FVector Start = InRayOrigin;
 	FVector End = InRayEnd;
 
 	// Cache a copy of the world pointer
-	UWorld* World = ViewportClient->GetWorld();
+	UWorld* World = GetWorld();
 
 	static FName TraceTag = FName(TEXT("LandscapeTrace"));
 	TArray<FHitResult> Results;
@@ -2521,8 +2509,8 @@ void FEdModeLandscape::ForceRealTimeViewports(const bool bEnable, const bool bSt
 					Viewport.SetRealtime(bEnable, bStoreCurrentState);
 
 					// @todo vreditor: Force game view to true in VREditor since we can't use hitproxies and debug objects yet
-					UVREditorMode* VREditorMode = GEditor->GetEditorWorldManager()->GetEditorWorldWrapper(Viewport.GetWorld())->GetVREditorMode();
-					if (VREditorMode != nullptr && VREditorMode->IsActive())
+					UVREditorMode* VREditorMode = Cast<UVREditorMode>( GEditor->GetEditorWorldExtensionsManager()->GetEditorWorldExtensions( GetWorld() )->FindExtension( UVREditorMode::StaticClass() ) );
+					if( VREditorMode != nullptr && VREditorMode->IsActive())
 					{
 						Viewport.SetGameView(true);
 					} 
diff --git a/Engine/Source/Editor/LandscapeEditor/Private/LandscapeEdMode.h b/Engine/Source/Editor/LandscapeEditor/Private/LandscapeEdMode.h
--- a/Engine/Source/Editor/LandscapeEditor/Private/LandscapeEdMode.h
+++ b/Engine/Source/Editor/LandscapeEditor/Private/LandscapeEdMode.h
@@ -422,7 +422,7 @@ public:
 	bool LandscapePlaneTrace(FEditorViewportClient* ViewportClient, int32 MouseX, int32 MouseY, const FPlane& Plane, FVector& OutHitLocation);
 
 	/** Trace under the specified laser start and direction and return the landscape hit and the hit location (in landscape quad space) */
-	bool LandscapeTrace(const FEditorViewportClient* ViewportClient, const FVector& InRayOrigin, const FVector& InRayEnd, FVector& OutHitLocation);
+	bool LandscapeTrace(const FVector& InRayOrigin, const FVector& InRayEnd, FVector& OutHitLocation);
 
 	void SetCurrentToolMode(FName ToolModeName, bool bRestoreCurrentTool = true);
 
@@ -450,7 +450,7 @@ public:
 	/** Called when the user presses a button on their motion controller device */
 	void OnVRAction(FEditorViewportClient& ViewportClient, UViewportInteractor* Interactor, const FViewportActionKeyInput& Action, bool& bOutIsInputCaptured, bool& bWasHandled);
 
-	void OnVRHoverUpdate(FEditorViewportClient& ViewportClient, UViewportInteractor* Interactor, FVector& HoverImpactPoint, bool& bWasHandled);
+	void OnVRHoverUpdate(UViewportInteractor* Interactor, FVector& HoverImpactPoint, bool& bWasHandled);
 
 	/** Handle notification that visible levels may have changed and we should update the editable landscapes list */
 	void HandleLevelsChanged();
diff --git a/Engine/Source/Editor/MeshPaint/MeshPaint.Build.cs b/Engine/Source/Editor/MeshPaint/MeshPaint.Build.cs
--- a/Engine/Source/Editor/MeshPaint/MeshPaint.Build.cs
+++ b/Engine/Source/Editor/MeshPaint/MeshPaint.Build.cs
@@ -35,6 +35,14 @@ public class MeshPaint : ModuleRules
             }
         );
 
+		CircularlyReferencedDependentModules.AddRange(
+			new string[]
+			{
+				"ViewportInteraction",
+				"VREditor"
+			}
+		);
+
 		PrivateIncludePathModuleNames.AddRange(
 			new string[]
 			{
diff --git a/Engine/Source/Editor/MeshPaint/Private/MeshPaintEdMode.cpp b/Engine/Source/Editor/MeshPaint/Private/MeshPaintEdMode.cpp
--- a/Engine/Source/Editor/MeshPaint/Private/MeshPaintEdMode.cpp
+++ b/Engine/Source/Editor/MeshPaint/Private/MeshPaintEdMode.cpp
@@ -27,7 +27,6 @@
 
 #include "StaticMeshResources.h"
 #include "VREditorMode.h"
-#include "IVREditorModule.h"
 #include "ScopedTransaction.h"
 #include "MeshPaintRendering.h"
 #include "Toolkits/ToolkitManager.h"
@@ -50,7 +49,6 @@
 #include "ViewportWorldInteraction.h"
 #include "ViewportInteractableInterface.h"
 #include "VREditorInteractor.h"
-#include "EditorWorldManager.h"
 #include "UniquePtr.h"
 
 #define LOCTEXT_NAMESPACE "MeshPaint_Mode"
@@ -257,18 +255,18 @@ void FEdModeMeshPaint::Enter()
 		ApplyOrRemoveForceBestLOD(/*bApply=*/ true, PaintingMeshLODIndex);
 	}
 
-	TSharedPtr<FEditorWorldWrapper> EditorWorld = GEditor->GetEditorWorldManager()->GetEditorWorldWrapper( GetWorld() );
+	TSharedPtr<FEditorWorldExtensionCollection> EditorWorld = GEditor->GetEditorWorldExtensionsManager()->GetEditorWorldExtensions( GetWorld() );
 	if(EditorWorld.IsValid())
 	{
 		// Register to find out about VR input events
-		UViewportWorldInteraction* WorldInteraction = EditorWorld->GetViewportWorldInteraction();
-		if(WorldInteraction != nullptr)
+		UViewportWorldInteraction* ViewportWorldInteraction = Cast<UViewportWorldInteraction>(GEditor->GetEditorWorldExtensionsManager()->GetEditorWorldExtensions(GetWorld())->FindExtension(UViewportWorldInteraction::StaticClass()));
+		if(ViewportWorldInteraction != nullptr)
 		{
-			WorldInteraction->OnViewportInteractionInputAction().RemoveAll(this);
-			WorldInteraction->OnViewportInteractionInputAction().AddRaw(this, &FEdModeMeshPaint::OnVRAction);
+			ViewportWorldInteraction->OnViewportInteractionInputAction().RemoveAll(this);
+			ViewportWorldInteraction->OnViewportInteractionInputAction().AddRaw(this, &FEdModeMeshPaint::OnVRAction);
 
 			// Hide the VR transform gizmo while we're in mesh paint mode.  It sort of gets in the way of painting.
-			WorldInteraction->SetTransformGizmoVisible(false);
+			ViewportWorldInteraction->SetTransformGizmoVisible(false);
 		}
 	}
 }
@@ -276,16 +274,15 @@ void FEdModeMeshPaint::Enter()
 /** FEdMode: Called when the mode is exited */
 void FEdModeMeshPaint::Exit()
 {
-	if( IVREditorModule::IsAvailable() )
 	{
-		UViewportWorldInteraction* ViewpportWorldInteraction = GEditor->GetEditorWorldManager()->GetEditorWorldWrapper(GetWorld())->GetViewportWorldInteraction();
-		if(ViewpportWorldInteraction != nullptr )
+		UViewportWorldInteraction* ViewportWorldInteraction = Cast<UViewportWorldInteraction>(GEditor->GetEditorWorldExtensionsManager()->GetEditorWorldExtensions(GetWorld())->FindExtension(UViewportWorldInteraction::StaticClass()));
+		if(ViewportWorldInteraction != nullptr )
 		{
 			// Restore the transform gizmo visibility
-			ViewpportWorldInteraction->SetTransformGizmoVisible( true );
+			ViewportWorldInteraction->SetTransformGizmoVisible( true );
 
 			// Unregister from event handlers
-			ViewpportWorldInteraction->OnViewportInteractionInputAction().RemoveAll( this );
+			ViewportWorldInteraction->OnViewportInteractionInputAction().RemoveAll( this );
 		}
 	}
 
@@ -2749,11 +2746,9 @@ void FEdModeMeshPaint::Render( const FSceneView* View, FViewport* Viewport, FPri
 		static TArray<FPaintRay> PaintRays;
 		PaintRays.Reset();
 		
-		UVREditorMode* VREditorMode = GEditor->GetEditorWorldManager()->GetEditorWorldWrapper( Viewport->GetClient()->GetWorld() )->GetVREditorMode();
-
+		UVREditorMode* VREditorMode = Cast<UVREditorMode>( GEditor->GetEditorWorldExtensionsManager()->GetEditorWorldExtensions( GetWorld() )->FindExtension( UVREditorMode::StaticClass() ) );
 		// Check to see if VREditorMode is active. If so, we're painting with interactor
 		bool bIsInVRMode = false;
-		if ( IVREditorModule::IsAvailable() )
 		{
 			if ( VREditorMode != nullptr && VREditorMode->IsFullyInitialized() && VREditorMode->IsActive())
 			{
@@ -3676,7 +3671,7 @@ EMeshPaintAction::Type FEdModeMeshPaint::GetPaintAction(FViewport* InViewport)
 		// When painting using VR, allow the modifier button to activate Erase mode
 		if( PaintingWithInteractorInVR )
 		{
-			UVREditorMode* VREditorMode = GEditor->GetEditorWorldManager()->GetEditorWorldWrapper( GetWorld() )->GetVREditorMode();
+			UVREditorMode* VREditorMode = Cast<UVREditorMode>( GEditor->GetEditorWorldExtensionsManager()->GetEditorWorldExtensions( GetWorld() )->FindExtension( UVREditorMode::StaticClass() ) );
 			UVREditorInteractor* VRInteractor = Cast<UVREditorInteractor>( PaintingWithInteractorInVR );
 			if( VREditorMode != nullptr && VRInteractor != nullptr )
 			{
@@ -5428,7 +5423,7 @@ TWeakObjectPtr<AActor> FEdModeMeshPaint::GetEditingActor() const
 void FEdModeMeshPaint::OnVRAction( class FEditorViewportClient& ViewportClient, UViewportInteractor* Interactor, 
 	const FViewportActionKeyInput& Action, bool& bOutIsInputCaptured, bool& bWasHandled )
 {
-	UVREditorMode* VREditorMode = GEditor->GetEditorWorldManager()->GetEditorWorldWrapper(ViewportClient.GetWorld())->GetVREditorMode();
+	UVREditorMode* VREditorMode = Cast<UVREditorMode>( GEditor->GetEditorWorldExtensionsManager()->GetEditorWorldExtensions( GetWorld() )->FindExtension( UVREditorMode::StaticClass() ) );
 	UVREditorInteractor* VRInteractor = Cast<UVREditorInteractor>( Interactor );
 	if( VREditorMode != nullptr && VREditorMode->IsActive() && VRInteractor != nullptr )
 	{
diff --git a/Engine/Source/Editor/Sequencer/Private/SequencerEdMode.cpp b/Engine/Source/Editor/Sequencer/Private/SequencerEdMode.cpp
--- a/Engine/Source/Editor/Sequencer/Private/SequencerEdMode.cpp
+++ b/Engine/Source/Editor/Sequencer/Private/SequencerEdMode.cpp
@@ -16,8 +16,8 @@
 #include "SubtitleManager.h"
 #include "SequencerMeshTrail.h"
 #include "SequencerKeyActor.h"
+#include "EditorWorldExtension.h"
 #include "ViewportWorldInteraction.h"
-#include "EditorWorldManager.h"
 
 const FEditorModeID FSequencerEdMode::EM_SequencerMode(TEXT("EM_SequencerMode"));
 
@@ -545,8 +545,7 @@ void FSequencerEdMode::DrawTracks3D(FPrimitiveDrawInterface* PDI)
 							});
 							if (TrailPtr == nullptr)
 							{
-								UViewportWorldInteraction* ViewportWorldInteraction = GEditor->GetEditorWorldManager()->GetEditorWorldWrapper(GEditor->GetActiveViewport()->GetClient()->GetWorld())->GetViewportWorldInteraction();
-								ASequencerMeshTrail* TrailActor = ViewportWorldInteraction->SpawnTransientSceneActor<ASequencerMeshTrail>(TEXT("SequencerMeshTrail"), true);
+								ASequencerMeshTrail* TrailActor = UViewportWorldInteraction::SpawnTransientSceneActor<ASequencerMeshTrail>(GetWorld(), TEXT("SequencerMeshTrail"), true);
 								FMeshTrailData MeshTrail = FMeshTrailData(TransformTrack, TrailActor);
 								MeshTrails.Add(MeshTrail);
 							}
diff --git a/Engine/Source/Editor/Sequencer/Private/SequencerMeshTrail.cpp b/Engine/Source/Editor/Sequencer/Private/SequencerMeshTrail.cpp
--- a/Engine/Source/Editor/Sequencer/Private/SequencerMeshTrail.cpp
+++ b/Engine/Source/Editor/Sequencer/Private/SequencerMeshTrail.cpp
@@ -3,7 +3,7 @@
 #include "SequencerMeshTrail.h"
 #include "SequencerKeyActor.h"
 #include "ViewportWorldInteraction.h"
-#include "EditorWorldManager.h"
+#include "EditorWorldExtension.h"
 #include "Materials/MaterialInstance.h" 
 
 ASequencerMeshTrail::ASequencerMeshTrail()
@@ -36,8 +36,7 @@ void ASequencerMeshTrail::AddKeyMeshActor(float KeyTime, const FTransform KeyTra
 	// If we don't currently have an actor for this time, create one
 	if (KeyPtr == nullptr)
 	{
-		UViewportWorldInteraction* ViewportWorldInteraction = GEditor->GetEditorWorldManager()->GetEditorWorldWrapper(GEditor->GetActiveViewport()->GetClient()->GetWorld())->GetViewportWorldInteraction();
-		KeyMeshActor = ViewportWorldInteraction->SpawnTransientSceneActor<ASequencerKeyActor>(TEXT("KeyMesh"), false);
+		KeyMeshActor = UViewportWorldInteraction::SpawnTransientSceneActor<ASequencerKeyActor>(GetWorld(), TEXT("KeyMesh"), false);
 		KeyMeshActor->SetMobility(EComponentMobility::Movable);
 		KeyMeshActor->SetActorTransform(KeyTransform);
 		KeyMeshActor->SetKeyData(TrackSection, KeyTime);
diff --git a/Engine/Source/Editor/Sequencer/Sequencer.Build.cs b/Engine/Source/Editor/Sequencer/Sequencer.Build.cs
--- a/Engine/Source/Editor/Sequencer/Sequencer.Build.cs
+++ b/Engine/Source/Editor/Sequencer/Sequencer.Build.cs
@@ -39,6 +39,13 @@ public class Sequencer : ModuleRules
 			}
 		);
 
+		CircularlyReferencedDependentModules.AddRange(
+			new string[]
+			{
+				"ViewportInteraction",
+			}
+		);
+
 		PrivateIncludePathModuleNames.AddRange(
 			new string[] {
                 "ContentBrowser",
diff --git a/Engine/Source/Editor/UnrealEd/Classes/Editor/EditorEngine.h b/Engine/Source/Editor/UnrealEd/Classes/Editor/EditorEngine.h
--- a/Engine/Source/Editor/UnrealEd/Classes/Editor/EditorEngine.h
+++ b/Engine/Source/Editor/UnrealEd/Classes/Editor/EditorEngine.h
@@ -1819,7 +1819,7 @@ public:
 	/**
 	*  Returns the Editors world manager instance.
 	*/
-	TSharedRef<class FEditorWorldManager> GetEditorWorldManager() { return EditorWorldManager.ToSharedRef(); }
+	TSharedRef<class FEditorWorldExtensionManager> GetEditorWorldExtensionsManager() { return EditorWorldExtensionsManager.ToSharedRef(); }
 
 	// Editor specific
 
@@ -2824,7 +2824,7 @@ private:
 	TSharedPtr<class FTimerManager> TimerManager;
 
 	/** Manager that holds all the worlds contexts with viewport interaction */
-	TSharedPtr<class FEditorWorldManager> EditorWorldManager;
+	TSharedPtr<class FEditorWorldExtensionManager> EditorWorldExtensionsManager;
 
 	/** The output log -> message log redirector for use during PIE */
 	TSharedPtr<class FOutputLogErrorsToMessageLogProxy> OutputLogErrorsToMessageLogProxyPtr;
diff --git a/Engine/Source/Editor/UnrealEd/Private/EditorEngine.cpp b/Engine/Source/Editor/UnrealEd/Private/EditorEngine.cpp
--- a/Engine/Source/Editor/UnrealEd/Private/EditorEngine.cpp
+++ b/Engine/Source/Editor/UnrealEd/Private/EditorEngine.cpp
@@ -134,7 +134,7 @@
 #include "Components/SkeletalMeshComponent.h"
 #include "EngineModule.h"
 
-#include "EditorWorldManager.h"
+#include "EditorWorldExtension.h"
 
 #if PLATFORM_WINDOWS
 	#include "WindowsHWrapper.h"
@@ -613,7 +613,7 @@ void UEditorEngine::InitEditor(IEngineLoop* InEngineLoop)
 	TimerManager = MakeShareable(new FTimerManager());
 
 	// create the editor world manager
-	EditorWorldManager = MakeShareable(new FEditorWorldManager());
+	EditorWorldExtensionsManager = MakeShareable(new FEditorWorldExtensionManager());
 
 	// Settings.
 	FBSPOps::GFastRebuild = 0;
@@ -1700,7 +1700,7 @@ void UEditorEngine::Tick( float DeltaSeconds, bool bIdleMode )
 	}
 
 	// Updates the ViewportWorldInteraction
-	EditorWorldManager->Tick( DeltaSeconds );
+	EditorWorldExtensionsManager->Tick( DeltaSeconds );
 
 	bool bIsMouseOverAnyLevelViewport = false;
 
@@ -7008,4 +7008,4 @@ void UEditorEngine::AutomationLoadMap(const FString& MapName, FString* OutError)
 	return;
 }
 
-#undef LOCTEXT_NAMESPACE 
+#undef LOCTEXT_NAMESPACE 
\ No newline at end of file
diff --git a/Engine/Source/Editor/UnrealEd/Private/EditorViewportClient.cpp b/Engine/Source/Editor/UnrealEd/Private/EditorViewportClient.cpp
--- a/Engine/Source/Editor/UnrealEd/Private/EditorViewportClient.cpp
+++ b/Engine/Source/Editor/UnrealEd/Private/EditorViewportClient.cpp
@@ -43,6 +43,7 @@
 #include "ComponentRecreateRenderStateContext.h"
 #include "EditorBuildUtils.h"
 #include "AudioDevice.h"
+#include "EditorWorldExtension.h"
 
 #define LOCTEXT_NAMESPACE "EditorViewportClient"
 
@@ -2270,6 +2271,12 @@ bool FEditorViewportClient::InputKey(FViewport* InViewport, int32 ControllerId,
 		return true;
 	}
 
+	FEditorWorldExtensionCollection& EditorWorldExtensionCollection = *GEditor->GetEditorWorldExtensionsManager()->GetEditorWorldExtensions( GetWorld() );
+	if( EditorWorldExtensionCollection.InputKey(this, Viewport, Key, Event))
+	{
+		return true;
+	}
+
 	FInputEventState InputState(InViewport, Key, Event);
 
 	bool bHandled = false;
diff --git a/Engine/Source/Editor/UnrealEd/Private/EditorWorldExtension.cpp b/Engine/Source/Editor/UnrealEd/Private/EditorWorldExtension.cpp
new file mode 100644
--- /dev/null
+++ b/Engine/Source/Editor/UnrealEd/Private/EditorWorldExtension.cpp
@@ -0,0 +1,257 @@
+// Copyright 1998-2017 Epic Games, Inc. All Rights Reserved.
+
+#include "EditorWorldExtension.h"
+#include "EngineGlobals.h"
+#include "Editor.h"
+
+/************************************************************************/
+/* UEditorExtension		                                                */
+/************************************************************************/
+UEditorWorldExtension::UEditorWorldExtension() :
+	Super(),
+	OwningWorld( nullptr )
+{
+
+}
+
+UEditorWorldExtension::~UEditorWorldExtension()
+{
+	OwningWorld = nullptr;
+}
+
+bool UEditorWorldExtension::InputKey( FEditorViewportClient* InViewportClient, FViewport* Viewport, FKey Key, EInputEvent Event )
+{
+	return false;
+}
+
+bool UEditorWorldExtension::InputAxis( FEditorViewportClient* InViewportClient, FViewport* Viewport, int32 ControllerId, FKey Key, float Delta, float DeltaTime )
+{
+	return false;
+}
+
+UWorld* UEditorWorldExtension::GetWorld() const
+{
+	return OwningWorld;
+}
+
+void UEditorWorldExtension::SetWorld( UWorld* InitWorld )
+{
+	OwningWorld = InitWorld;
+}
+
+/************************************************************************/
+/* FEditorWorldExtensionCollection                                      */
+/************************************************************************/
+FEditorWorldExtensionCollection::FEditorWorldExtensionCollection( FWorldContext& InWorldContext ) :
+	WorldContext ( InWorldContext )
+{
+
+}
+
+FEditorWorldExtensionCollection::~FEditorWorldExtensionCollection()
+{
+	EditorExtensions.Empty();
+}
+
+UWorld* FEditorWorldExtensionCollection::GetWorld() const
+{
+	return WorldContext.World();
+}
+
+FWorldContext& FEditorWorldExtensionCollection::GetWorldContext() const
+{
+	return WorldContext;
+}
+
+UEditorWorldExtension* FEditorWorldExtensionCollection::AddExtension( TSubclassOf<UEditorWorldExtension> EditorExtensionClass )
+{
+	UEditorWorldExtension* CreatedExtension = EditorExtensionClass.GetDefaultObject();
+	AddExtension( CreatedExtension );
+	return CreatedExtension;
+}
+
+void FEditorWorldExtensionCollection::AddExtension( UEditorWorldExtension* EditorExtension )
+{
+	check( EditorExtension != nullptr );
+
+	const int32 ExistingExtensionIndex = EditorExtensions.IndexOfByPredicate(
+		[ &EditorExtension ]( const FEditorExtensionTuple& Element ) -> bool
+		{ 
+			return Element.Get<0>() == EditorExtension;
+		} 
+	);
+
+	if( ExistingExtensionIndex != INDEX_NONE )
+	{
+		FEditorExtensionTuple& EditorExtensionTuple = EditorExtensions[ ExistingExtensionIndex ];
+		int32& RefCount = EditorExtensionTuple.Get<1>();
+		++RefCount;
+	}
+	else
+	{
+		const int32 InitialRefCount = 1;
+		EditorExtensions.Add( FEditorExtensionTuple( EditorExtension, InitialRefCount ) );
+
+		EditorExtension->SetWorld( GetWorld() );
+		EditorExtension->Init();
+	}
+}
+
+void FEditorWorldExtensionCollection::RemoveExtension( UEditorWorldExtension* EditorExtension )
+{
+	check( EditorExtension != nullptr );
+
+	const int32 ExistingExtensionIndex = EditorExtensions.IndexOfByPredicate(
+		[ &EditorExtension ]( const FEditorExtensionTuple& Element ) -> bool 
+		{ 
+			return Element.Get<0>() == EditorExtension;
+		} 
+	);
+	if( ensure( ExistingExtensionIndex != INDEX_NONE ) )
+	{
+		FEditorExtensionTuple& EditorExtensionTuple = EditorExtensions[ ExistingExtensionIndex ];
+		int32& RefCount = EditorExtensionTuple.Get<1>();
+		--RefCount;
+
+		if( RefCount == 0 )
+		{
+			EditorExtensions.RemoveAt( ExistingExtensionIndex );
+			EditorExtension->Shutdown();
+		}
+	}
+}
+
+UEditorWorldExtension* FEditorWorldExtensionCollection::FindExtension( TSubclassOf<UEditorWorldExtension> EditorExtensionClass )
+{
+	UEditorWorldExtension* ResultExtension = nullptr;
+	for( FEditorExtensionTuple& EditorExtensionTuple : EditorExtensions )
+	{
+		UEditorWorldExtension* EditorExtension = EditorExtensionTuple.Get<0>();
+		if( EditorExtension->GetClass() == EditorExtensionClass )
+		{
+			ResultExtension = EditorExtension;
+			break;
+		}
+	}
+
+	return ResultExtension;
+}
+
+void FEditorWorldExtensionCollection::Tick( float DeltaSeconds )
+{
+	for( FEditorExtensionTuple& EditorExtensionTuple : EditorExtensions )
+	{
+		UEditorWorldExtension* EditorExtension = EditorExtensionTuple.Get<0>();
+		EditorExtension->Tick( DeltaSeconds );
+	}
+}
+
+bool FEditorWorldExtensionCollection::InputKey( FEditorViewportClient* InViewportClient, FViewport* Viewport, FKey Key, EInputEvent Event )
+{
+	bool bHandled = false;
+
+	for( FEditorExtensionTuple& EditorExtensionTuple : EditorExtensions )
+	{
+		UEditorWorldExtension* EditorExtension = EditorExtensionTuple.Get<0>();
+		bHandled |= EditorExtension->InputKey( InViewportClient, Viewport, Key, Event );
+	}
+
+	return bHandled;
+}
+
+bool FEditorWorldExtensionCollection::InputAxis( FEditorViewportClient* InViewportClient, FViewport* Viewport, int32 ControllerId, FKey Key, float Delta, float DeltaTime )
+{
+	bool bHandled = false;
+
+	for( FEditorExtensionTuple& EditorExtensionTuple : EditorExtensions )
+	{
+		UEditorWorldExtension* EditorExtension = EditorExtensionTuple.Get<0>();
+		bHandled |= EditorExtension->InputAxis( InViewportClient, Viewport, ControllerId, Key, Delta, DeltaTime );
+	}
+
+	return bHandled;
+}
+
+void FEditorWorldExtensionCollection::AddReferencedObjects( FReferenceCollector& Collector )
+{
+	for( FEditorExtensionTuple& EditorExtensionTuple : EditorExtensions )
+	{
+		UEditorWorldExtension* EditorExtension = EditorExtensionTuple.Get<0>();
+		Collector.AddReferencedObject( EditorExtension );
+	}
+}
+
+/************************************************************************/
+/* FEditorWorldExtensionManager                                         */
+/************************************************************************/
+FEditorWorldExtensionManager::FEditorWorldExtensionManager()
+{
+	if( GEngine )
+	{
+		GEngine->OnWorldContextDestroyed().AddRaw( this, &FEditorWorldExtensionManager::OnWorldContextRemove );
+	}
+}
+
+FEditorWorldExtensionManager::~FEditorWorldExtensionManager()
+{
+	if( GEngine )
+	{
+		GEngine->OnWorldContextDestroyed().RemoveAll( this );
+	}
+
+	EditorWorldExtensionCollection.Empty();
+}
+
+TSharedPtr<FEditorWorldExtensionCollection> FEditorWorldExtensionManager::GetEditorWorldExtensions( const UWorld* InWorld )
+{
+	// Try to find this world in the map and return it or create and add one if nothing found
+	TSharedPtr<FEditorWorldExtensionCollection> Result;
+	if( InWorld )
+	{
+		TSharedPtr<FEditorWorldExtensionCollection>* FoundWorld = EditorWorldExtensionCollection.Find( InWorld->GetUniqueID() );
+		if( FoundWorld != nullptr )
+		{
+			Result = *FoundWorld;
+		}
+		else
+		{
+			FWorldContext* WorldContext = GEditor->GetWorldContextFromWorld( InWorld );
+			Result = OnWorldContextAdd( *WorldContext );
+		}
+	}
+	return Result;
+}
+
+TSharedPtr<FEditorWorldExtensionCollection> FEditorWorldExtensionManager::OnWorldContextAdd( FWorldContext& InWorldContext )
+{
+	//Only add editor type world to the map
+	UWorld* World = InWorldContext.World();
+	TSharedPtr<FEditorWorldExtensionCollection> Result;
+	if( World && ( InWorldContext.WorldType == EWorldType::Editor || InWorldContext.WorldType == EWorldType::EditorPreview || InWorldContext.WorldType == EWorldType::GamePreview ) )
+	{
+		TSharedPtr<FEditorWorldExtensionCollection> ExtensionCollection( new FEditorWorldExtensionCollection( InWorldContext ) );
+		Result = ExtensionCollection;
+		EditorWorldExtensionCollection.Add( World->GetUniqueID(), Result );
+
+	}
+	return Result;
+}
+
+void FEditorWorldExtensionManager::OnWorldContextRemove( FWorldContext& InWorldContext )
+{
+	UWorld* World = InWorldContext.World();
+	if(World)
+	{
+		EditorWorldExtensionCollection.Remove( World->GetUniqueID() );
+	}
+}
+
+void FEditorWorldExtensionManager::Tick( float DeltaSeconds )
+{
+	// Tick all the collections
+	for( auto& Collection : EditorWorldExtensionCollection )
+	{
+		FEditorWorldExtensionCollection& ExtensionCollection = *Collection.Value;
+		ExtensionCollection.Tick( DeltaSeconds );
+	}
+}
\ No newline at end of file
diff --git a/Engine/Source/Editor/UnrealEd/Private/EditorWorldManager.cpp b/Engine/Source/Editor/UnrealEd/Private/EditorWorldManager.cpp
deleted file mode 100644
diff --git a/Engine/Source/Editor/UnrealEd/Public/EditorWorldExtension.h b/Engine/Source/Editor/UnrealEd/Public/EditorWorldExtension.h
new file mode 100644
--- /dev/null
+++ b/Engine/Source/Editor/UnrealEd/Public/EditorWorldExtension.h
@@ -0,0 +1,148 @@
+// Copyright 1998-2017 Epic Games, Inc. All Rights Reserved.
+
+#pragma once
+
+#include "CoreMinimal.h"
+#include "UObject/GCObject.h"
+#include "EditorWorldExtension.generated.h"
+
+struct FWorldContext;
+
+UCLASS()
+class UNREALED_API UEditorWorldExtension : public UObject
+{
+	GENERATED_BODY()
+
+	friend class FEditorWorldExtensionCollection;
+public:
+	
+	/** Default constructor */
+	UEditorWorldExtension();
+
+	/** Default destructor */
+	virtual ~UEditorWorldExtension();
+
+	/** Initialize extension */
+	virtual void Init() {};
+
+	/** Shut down extension when world is destroyed */
+	virtual void Shutdown() {};
+
+	/** Give base class the chance to tick */
+	virtual void Tick( float DeltaSeconds ) {};
+
+	virtual bool InputKey( FEditorViewportClient* InViewportClient, FViewport* Viewport, FKey Key, EInputEvent Event);
+	virtual bool InputAxis( FEditorViewportClient* InViewportClient, FViewport* Viewport, int32 ControllerId, FKey Key, float Delta, float DeltaTime);
+
+	/** Gets the world owning this extension */
+	virtual UWorld* GetWorld() const override;
+
+private:
+
+	/** Let the FEditorWorldExtensionCollection set the world of this extension before init */
+	void SetWorld( UWorld* InitWorld );
+
+	/** The world that is owning this extension */
+	UWorld* OwningWorld;
+};
+
+/**
+ * Holds a collection of UEditorExtension
+ */
+class UNREALED_API FEditorWorldExtensionCollection : public FGCObject
+{
+public:
+
+	/** Default constructor */
+	FEditorWorldExtensionCollection( FWorldContext& InWorldContext );
+
+	/** Default destructor */
+	virtual ~FEditorWorldExtensionCollection();
+
+	/** Gets the world from the world context */
+	UWorld* GetWorld() const;
+
+	/** Gets the world context */
+	FWorldContext& GetWorldContext() const;
+
+	/** 
+	 * Adds an extension to the collection
+	 * @param	EditorExtensionClass	The subclass of UEditorExtension that will be created, initialized and added to the collection.
+	 * @return							The created UEditorExtension.
+	 */
+	UEditorWorldExtension* AddExtension( TSubclassOf<UEditorWorldExtension> EditorExtensionClass );
+
+	/** 
+	 * Adds an extension to the collection
+	 * @param	EditorExtension			The UEditorExtension that will be created, initialized and added to the collection.
+	 */
+	void AddExtension( UEditorWorldExtension* EditorExtension );
+	
+	/** 
+	 * Removes an extension from the collection and calls Shutdown() on the extension
+	 * @param	EditorExtension			The UEditorExtension to remove.  It must already have been added.
+	 */
+	void RemoveExtension( UEditorWorldExtension* EditorExtension );
+
+	/**
+	 * Find an extension based on the class
+	 * @param	EditorExtensionClass	The class to find an extension with
+	 * @return							The first extension that is found based on class
+	 */
+	UEditorWorldExtension* FindExtension( TSubclassOf<UEditorWorldExtension> EditorExtensionClass );
+
+	/** Ticks all extensions */
+	void Tick( float DeltaSeconds );
+
+	/** Notifies all extensions of keyboard input */
+	bool InputKey( FEditorViewportClient* InViewportClient, FViewport* Viewport, FKey Key, EInputEvent Event);
+
+	/** Notifies all extensions of axis movement */
+	bool InputAxis( FEditorViewportClient* InViewportClient, FViewport* Viewport, int32 ControllerId, FKey Key, float Delta, float DeltaTime);
+
+	// FGCObject
+	virtual void AddReferencedObjects( FReferenceCollector& Collector ) override;
+	// End FGCObject
+
+private:
+
+	/** World context */
+	FWorldContext& WorldContext;
+
+	/** List of extensions along with their reference count.  Extensions will only be truly removed and Shutdown() after their
+	    reference count drops to zero. */
+	typedef TTuple<UEditorWorldExtension*, int32> FEditorExtensionTuple;
+	TArray<FEditorExtensionTuple> EditorExtensions;
+};
+
+/**
+ * Holds a map of extension collections paired with worlds
+ */
+class UNREALED_API FEditorWorldExtensionManager	
+{
+public:
+
+	/** Default constructor */
+	FEditorWorldExtensionManager();
+
+	/** Default destructor */
+	virtual ~FEditorWorldExtensionManager();
+
+	/** Gets the editor world wrapper that is found with the world passed.
+	 * Adds one for this world if there was non found. */
+	TSharedPtr<FEditorWorldExtensionCollection> GetEditorWorldExtensions(const UWorld* InWorld);
+
+	/** Ticks all the collections */
+	void Tick( float DeltaSeconds );
+
+private:
+
+	/** Adds a new editor world wrapper when a new world context was created */
+	TSharedPtr<FEditorWorldExtensionCollection> OnWorldContextAdd(FWorldContext& InWorldContext);
+
+	/** Adds a new editor world wrapper when a new world context was created */
+	void OnWorldContextRemove(FWorldContext& InWorldContext);
+
+	/** Map of all the editor world maps */
+	TMap<uint32, TSharedPtr<FEditorWorldExtensionCollection>> EditorWorldExtensionCollection;
+};
\ No newline at end of file
diff --git a/Engine/Source/Editor/UnrealEd/Public/EditorWorldManager.h b/Engine/Source/Editor/UnrealEd/Public/EditorWorldManager.h
deleted file mode 100644
diff --git a/Engine/Source/Editor/UnrealEd/UnrealEd.Build.cs b/Engine/Source/Editor/UnrealEd/UnrealEd.Build.cs
--- a/Engine/Source/Editor/UnrealEd/UnrealEd.Build.cs
+++ b/Engine/Source/Editor/UnrealEd/UnrealEd.Build.cs
@@ -63,8 +63,6 @@ public class UnrealEd : ModuleRules
                 "PortalRpc",
                 "PortalServices",
                 "BlueprintNativeCodeGen",
-                "ViewportInteraction",
-                "VREditor",
                 "Persona",
             }
 		);
@@ -151,8 +149,6 @@ public class UnrealEd : ModuleRules
                 "PixelInspectorModule",
 				"MovieScene",
 				"MovieSceneTracks",
-                "ViewportInteraction",
-                "VREditor"
             }
 		);
 
@@ -299,4 +295,4 @@ public class UnrealEd : ModuleRules
 			Definitions.Add( "WITH_RECAST=0" );
 		}
 	}
-}
+}
\ No newline at end of file
diff --git a/Engine/Source/Editor/VREditor/Gizmo/VREditorTransformGizmo.cpp b/Engine/Source/Editor/VREditor/Gizmo/VREditorTransformGizmo.cpp
--- a/Engine/Source/Editor/VREditor/Gizmo/VREditorTransformGizmo.cpp
+++ b/Engine/Source/Editor/VREditor/Gizmo/VREditorTransformGizmo.cpp
@@ -139,10 +139,10 @@ ATransformGizmo::ATransformGizmo()
 }
 
 
-void ATransformGizmo::UpdateGizmo( const EGizmoHandleTypes GizmoType, const ECoordSystem GizmoCoordinateSpace, const FTransform& LocalToWorld, const FBox& LocalBounds, const FVector ViewLocation, bool bAllHandlesVisible, 
+void ATransformGizmo::UpdateGizmo( const EGizmoHandleTypes GizmoType, const ECoordSystem GizmoCoordinateSpace, const FTransform& LocalToWorld, const FBox& LocalBounds, const FVector ViewLocation, const float ScaleMultiplier, bool bAllHandlesVisible,
 	UActorComponent* DraggingHandle, const TArray< UActorComponent* >& HoveringOverHandles, const float GizmoHoverScale, const float GizmoHoverAnimationDuration )
 {
-	Super::UpdateGizmo( GizmoType, GizmoCoordinateSpace, LocalToWorld, LocalBounds, ViewLocation, bAllHandlesVisible,
+	Super::UpdateGizmo( GizmoType, GizmoCoordinateSpace, LocalToWorld, LocalBounds, ViewLocation, ScaleMultiplier, bAllHandlesVisible,
 		DraggingHandle, HoveringOverHandles, GizmoHoverScale, GizmoHoverAnimationDuration );
 
 	const float WorldScaleFactor = GetWorld()->GetWorldSettings()->WorldToMeters / 100.0f;
@@ -167,7 +167,7 @@ void ATransformGizmo::UpdateGizmo( const EGizmoHandleTypes GizmoType, const ECoo
 	// @todo vreditor: Should take FOV into account (especially in non-stereo/HMD mode)
 	const float WorldSpaceDistanceToGizmoBounds = FMath::Sqrt( WorldSpaceBounds.ComputeSquaredDistanceToPoint( ViewLocation ) );
 	const float GizmoScaleUpClose = VREd::GizmoScale->GetFloat();
-	const float GizmoScale( ( GizmoScaleUpClose + ( WorldSpaceDistanceToGizmoBounds / WorldScaleFactor ) * VREd::GizmoDistanceScaleFactor->GetFloat() ) * ScaleCompensationForTinyGizmo * WorldScaleFactor );
+	const float GizmoScale( ScaleMultiplier * ( GizmoScaleUpClose + ( WorldSpaceDistanceToGizmoBounds / WorldScaleFactor ) * VREd::GizmoDistanceScaleFactor->GetFloat() ) * ScaleCompensationForTinyGizmo * WorldScaleFactor );
 
 	// Update animation
 	float AnimationAlpha = GetAnimationAlpha();
diff --git a/Engine/Source/Editor/VREditor/Gizmo/VREditorTransformGizmo.h b/Engine/Source/Editor/VREditor/Gizmo/VREditorTransformGizmo.h
--- a/Engine/Source/Editor/VREditor/Gizmo/VREditorTransformGizmo.h
+++ b/Engine/Source/Editor/VREditor/Gizmo/VREditorTransformGizmo.h
@@ -23,7 +23,7 @@ public:
 	/** Called by the world interaction system after we've been spawned into the world, to allow
 	    us to create components and set everything up nicely for the selected objects that we'll be
 		used to manipulate */
-	virtual void UpdateGizmo( const EGizmoHandleTypes GizmoType, const ECoordSystem GizmoCoordinateSpace, const FTransform& LocalToWorld, const FBox& LocalBounds, const FVector ViewLocation, bool bAllHandlesVisible, 
+	virtual void UpdateGizmo( const EGizmoHandleTypes GizmoType, const ECoordSystem GizmoCoordinateSpace, const FTransform& LocalToWorld, const FBox& LocalBounds, const FVector ViewLocation, const float ScaleMultiplier, bool bAllHandlesVisible,
 		class UActorComponent* DraggingHandle, const TArray< UActorComponent* >& HoveringOverHandles, const float GizmoHoverScale, const float GizmoHoverAnimationDuration ) override;
 
 private:
diff --git a/Engine/Source/Editor/VREditor/Public/IVREditorModule.h b/Engine/Source/Editor/VREditor/Public/IVREditorModule.h
--- a/Engine/Source/Editor/VREditor/Public/IVREditorModule.h
+++ b/Engine/Source/Editor/VREditor/Public/IVREditorModule.h
@@ -65,5 +65,12 @@ public:
 	 * @return True if the VREditor is currently running 
 	 */
 	virtual bool IsVREditorModeActive() = 0;
+
+	/** 
+	* Get the current VREditor running
+	*
+	* @return The current VREditor running
+	*/
+	virtual class UVREditorMode* GetVRMode() = 0;
 };
 
diff --git a/Engine/Source/Editor/VREditor/Public/VREditorInteractor.h b/Engine/Source/Editor/VREditor/Public/VREditorInteractor.h
--- a/Engine/Source/Editor/VREditor/Public/VREditorInteractor.h
+++ b/Engine/Source/Editor/VREditor/Public/VREditorInteractor.h
@@ -45,7 +45,7 @@ public:
 	virtual FHitResult GetHitResultFromLaserPointer( TArray<AActor*>* OptionalListOfIgnoredActors = nullptr, const bool bIgnoreGizmos = false,
 		TArray<UClass*>* ObjectsInFrontOfGizmo = nullptr, const bool bEvenIfBlocked = false, const float LaserLengthOverride = 0.0f ) override;
 	virtual void ResetHoverState( const float DeltaTime ) override;
-	virtual void OnStartDragging( UActorComponent* ClickedComponent, const FVector& HitLocation, const bool bIsPlacingActors ) override;
+	virtual void OnStartDragging( const FVector& HitLocation, const bool bIsPlacingNewObjects ) override;
 
 	/** Returns the slide delta for pushing and pulling objects. Needs to be implemented by derived classes (e.g. touchpad for vive controller or scrollweel for mouse ) */
 	virtual float GetSlideDelta();
@@ -109,7 +109,7 @@ public:
 	float GetSelectAndMoveTriggerValue() const;
 	
 	/* ViewportInteractor overrides, checks if the laser is blocked by UI */
-	virtual bool GetIsLaserBlocked() override;
+	virtual bool GetIsLaserBlocked() const override;
 	
 protected:
 
@@ -169,4 +169,4 @@ protected:
 
 	/** Time that we either started showing or hiding help labels (for fade transitions) */
 	FTimespan HelpLabelShowOrHideStartTime;
-};
+};
\ No newline at end of file
diff --git a/Engine/Source/Editor/VREditor/Public/VREditorMode.h b/Engine/Source/Editor/VREditor/Public/VREditorMode.h
--- a/Engine/Source/Editor/VREditor/Public/VREditorMode.h
+++ b/Engine/Source/Editor/VREditor/Public/VREditorMode.h
@@ -8,6 +8,7 @@
 #include "InputCoreTypes.h"
 #include "Templates/SubclassOf.h"
 #include "Editor/UnrealEdTypes.h"
+#include "EditorWorldExtension.h"
 #include "ShowFlags.h"
 #include "Misc/App.h"
 #include "Widgets/SWindow.h"
@@ -46,19 +47,39 @@ namespace VRActionTypes
 }
 
 /**
- * VR Editor Mode.  Extends editor viewports with functionality for VR controls and object manipulation
+ * VR Editor Mode. Extends editor viewports with functionality for VR controls and object manipulation
  */
 UCLASS()
-class VREDITOR_API UVREditorMode : public UObject
+class VREDITOR_API UVREditorMode : public UEditorWorldExtension
 {
-
-	GENERATED_UCLASS_BODY()
+	GENERATED_BODY()
 
 public:
 
+	/** Default constructor */
+	UVREditorMode();
+
 	/** Cleans up this mode, called when the editor is shutting down */
 	virtual ~UVREditorMode();
 
+	/** Initialize the VREditor */
+	virtual void Init() override;
+
+	/** Shutdown the VREditor */
+	virtual void Shutdown() override;	
+	
+	/** When the user actually enters the VR Editor mode */
+	void Enter(const bool bReenteringVREditing);
+
+	/** When the user leaves the VR Editor mode */
+	void Exit(const bool bHMDShouldExitStereo);
+
+	/** Tick before the ViewportWorldInteraction is ticked */
+	void PreTick( const float DeltaTime );
+
+	/** Tick after the ViewportWorldInteraction is ticked */
+	void PostTick( const float DeltaTime );
+
 	/** Static: Sets whether we should actually use an HMD.  Call this before activating VR mode */
 	void SetActuallyUsingVR( const bool bShouldActuallyUseVR )
 	{
@@ -85,24 +106,6 @@ public:
 	{
 		return ExitType;
 	}
-
-	/** EditorWorldManager initializes the start of the editor */
-	void Init( class UViewportWorldInteraction* InViewportWorldInteraction );
-
-	/** EditorWorldManager shuts down the VREditor when closing the editor */
-	void Shutdown();
-
-	/** When the user actually enters the VR Editor mode */
-	void Enter(const bool bReenteringVREditing);
-
-	/** When the user leaves the VR Editor mode */
-	void Exit(const bool bHMDShouldExitStereo);
-
-	/** Tick before the ViewportWorldInteraction is ticked */
-	void PreTick( const float DeltaTime );
-
-	/** Tick after the ViewportWorldInteraction is ticked */
-	void Tick( const float DeltaTime );
 	
 	//bool InputKey( FEditorViewportClient* ViewportClient, FViewport* Viewport, FKey Key, EInputEvent Event );
 	//bool InputAxis( FEditorViewportClient* ViewportClient, FViewport* Viewport, int32 ControllerId, FKey Key, float Delta, float DeltaTime );
@@ -117,7 +120,7 @@ public:
 	AActor* GetAvatarMeshActor();
 
 	/** Gets the World from the ViewportWorldInteraction used by this mode */
-	UWorld* GetWorld() const;
+	virtual UWorld* GetWorld() const override;
 
 	/** Gets the world space transform of the calibrated VR room origin.  When using a seated VR device, this will feel like the
 	camera's world transform (before any HMD positional or rotation adjustments are applied.) */
@@ -182,19 +185,6 @@ public:
 		return 1.0f - ( ( Alpha * ( ( OvershootAmount + 1 ) * Alpha + OvershootAmount ) + 1 ) - 1.0f );
 	}
 
-	/** Spawns a transient actor that we can use in the editor world (templated for convenience) */
-	template<class T>
-	T* SpawnTransientSceneActor( const FString& ActorName, const bool bWithSceneComponent ) const
-	{
-		return CastChecked<T>( SpawnTransientSceneActor( T::StaticClass(), ActorName, bWithSceneComponent ) );
-	}
-
-	/** Spawns a transient actor that we can use in the editor world */
-	AActor* SpawnTransientSceneActor( TSubclassOf<class AActor> ActorClass, const FString& ActorName, const bool bWithSceneComponent ) const;
-
-	/** Destroys a transient actor we created earlier and nulls out the pointer */
-	void DestroyTransientActor( AActor* Actor ) const;
-
 	/** Gets access to the VR UI system (const) */
 	const class UVREditorUISystem& GetUISystem() const
 	{
@@ -264,6 +254,7 @@ public:
 		float NearClipPlane;
 		bool bRealTime;
 		float DragTriggerDistance;
+		float TransformGizmoScale;
 		bool bOnScreenMessages;
 		EHMDTrackingOrigin::Type TrackingOrigin;
 		float WorldToMetersScale;
@@ -272,20 +263,21 @@ public:
 
 		FSavedEditorState()
 			: ViewportType(LVT_Perspective),
-			ViewLocation(FVector::ZeroVector),
-			ViewRotation(FRotator::ZeroRotator),
-			ShowFlags(ESFIM_Editor),
-			bLockedPitch(false),
-			bGameView(false),
-			bAlwaysShowModeWidgetAfterSelectionChanges(false),
-			NearClipPlane(0.0f),
-			bRealTime(false),
-			DragTriggerDistance(0.0f),
-			bOnScreenMessages(false),
-			TrackingOrigin(EHMDTrackingOrigin::Eye),
-			WorldToMetersScale(100.0f),
-			bKeyAllEnabled(false),
-			AutoKeyMode()
+			  ViewLocation(FVector::ZeroVector),
+			  ViewRotation(FRotator::ZeroRotator),
+ 			  ShowFlags(ESFIM_Editor),
+			  bLockedPitch(false),
+			  bGameView(false),
+			  bAlwaysShowModeWidgetAfterSelectionChanges(false),
+			  NearClipPlane(0.0f),
+			  bRealTime(false),
+			  DragTriggerDistance(0.0f),
+			  bOnScreenMessages(false),
+			  TransformGizmoScale( 1.0f ),
+			  TrackingOrigin(EHMDTrackingOrigin::Eye),
+			  WorldToMetersScale(100.0f),
+			  bKeyAllEnabled(false),
+			  AutoKeyMode()
 		{
 		}
 	};
@@ -306,7 +298,7 @@ private:
 	void OnEditorClosed();
 
 	//Handles closing the VR mode by escape key
-	void InputKey(const FEditorViewportClient& InViewportClient, const FKey InKey, const EInputEvent InEvent, bool& bOutWasHandled);
+	virtual bool InputKey( FEditorViewportClient* InViewportClient, FViewport* Viewport, FKey Key, EInputEvent Event ) override;
 
 	/** Called when someone closes a standalone VR Editor window */
 	void OnVREditorWindowClosed( const TSharedRef<SWindow>& ClosedWindow );
@@ -413,10 +405,6 @@ protected:
 	// Interactors
 	//
 
-	/** The mouse cursor interactor (currently only available when not in VR) */
-	UPROPERTY()
-	class UMouseCursorInteractor* MouseCursorInteractor;
-
 	/** The right motion controller */
 	UPROPERTY()
 	class UVREditorMotionControllerInteractor* LeftHandInteractor; //@todo vreditor: Hardcoded interactors
@@ -467,9 +455,6 @@ private:
 	/** If this is the first tick or before */
 	bool bFirstTick;
 
-	/** If the coordinate system was in world space before entering (local) scale mode */
-	bool bWasInWorldSpaceBeforeScaleMode;
-
 	/** If this current mode is running */
 	bool bIsActive;
 };
diff --git a/Engine/Source/Editor/VREditor/Teleporter/VREditorAutoScaler.cpp b/Engine/Source/Editor/VREditor/Teleporter/VREditorAutoScaler.cpp
--- a/Engine/Source/Editor/VREditor/Teleporter/VREditorAutoScaler.cpp
+++ b/Engine/Source/Editor/VREditor/Teleporter/VREditorAutoScaler.cpp
@@ -38,7 +38,7 @@ void UVREditorAutoScaler::Scale( const float NewWorldToMetersScale )
 	UViewportWorldInteraction* WorldInteraction = &VRMode->GetWorldInteraction();
 
 	// Set the new world to meters scale.
-	const float OldWorldToMetersScale = WorldInteraction->GetViewportWorld()->GetWorldSettings()->WorldToMeters;
+	const float OldWorldToMetersScale = WorldInteraction->GetWorld()->GetWorldSettings()->WorldToMeters;
 	WorldInteraction->SetWorldToMetersScale( NewWorldToMetersScale );
 	WorldInteraction->SkipInteractiveWorldMovementThisFrame();
 
diff --git a/Engine/Source/Editor/VREditor/Teleporter/VREditorTeleporter.cpp b/Engine/Source/Editor/VREditor/Teleporter/VREditorTeleporter.cpp
--- a/Engine/Source/Editor/VREditor/Teleporter/VREditorTeleporter.cpp
+++ b/Engine/Source/Editor/VREditor/Teleporter/VREditorTeleporter.cpp
@@ -104,7 +104,7 @@ void UVREditorTeleporter::StartTeleport( UViewportInteractor* Interactor )
 	}
 }
 
-void UVREditorTeleporter::OnInteractorAction( class FEditorViewportClient& ViewportClient, UViewportInteractor* Interactor, const struct FViewportActionKeyInput& Action, bool& bOutIsInputCaptured, bool& bWasHandled )
+void UVREditorTeleporter::OnInteractorAction( FEditorViewportClient& ViewportClient, UViewportInteractor* Interactor, const struct FViewportActionKeyInput& Action, bool& bOutIsInputCaptured, bool& bWasHandled )
 {
 	if ( ( Action.ActionType == ViewportWorldActionTypes::SelectAndMove ||
 		Action.ActionType == ViewportWorldActionTypes::SelectAndMove_LightlyPressed ) )
diff --git a/Engine/Source/Editor/VREditor/UI/VREditorUISystem.cpp b/Engine/Source/Editor/VREditor/UI/VREditorUISystem.cpp
--- a/Engine/Source/Editor/VREditor/UI/VREditorUISystem.cpp
+++ b/Engine/Source/Editor/VREditor/UI/VREditorUISystem.cpp
@@ -243,7 +243,7 @@ void UVREditorUISystem::OnVRAction( FEditorViewportClient& ViewportClient, UView
 {
 	UVREditorMotionControllerInteractor* VREditorInteractor = Cast<UVREditorMotionControllerInteractor>( Interactor );
 	if( VREditorInteractor && 
-		!( VREditorInteractor->GetDraggingMode() == EViewportInteractionDraggingMode::ActorsWithGizmo ||
+		!( VREditorInteractor->GetDraggingMode() == EViewportInteractionDraggingMode::TransformablesWithGizmo ||
 		  VREditorInteractor->GetDraggingMode() == EViewportInteractionDraggingMode::Interactable ) )
 	{
 		if ( !bWasHandled )
@@ -277,11 +277,11 @@ void UVREditorUISystem::OnVRAction( FEditorViewportClient& ViewportClient, UView
 			if ( Action.ActionType == VRActionTypes::ConfirmRadialSelection )
 			{
 				const EViewportInteractionDraggingMode DraggingMode = Interactor->GetDraggingMode();
-				// Toggle the radial menu by pressing the touchpad or thumb stick
+				// Toggle the radial menu by pressing the touchpad or thumb stick, as long as we're not busy dragging something
 				if ( Action.Event == IE_Pressed && 
-					 ( DraggingMode != EViewportInteractionDraggingMode::ActorsFreely && DraggingMode != EViewportInteractionDraggingMode::World && 
-					   ( DraggingMode != EViewportInteractionDraggingMode::AssistingDrag &&
-						 Interactor->GetOtherInteractor() != nullptr ) ) )
+					 DraggingMode != EViewportInteractionDraggingMode::TransformablesFreely &&
+					 DraggingMode != EViewportInteractionDraggingMode::World &&
+					 DraggingMode != EViewportInteractionDraggingMode::AssistingDrag )
 				{
 					if ( !IsShowingRadialMenu(VREditorInteractor) )
 					{
@@ -549,7 +549,7 @@ void UVREditorUISystem::OnVRAction( FEditorViewportClient& ViewportClient, UView
 }
 
 
-void UVREditorUISystem::OnVRHoverUpdate( FEditorViewportClient& ViewportClient, UViewportInteractor* Interactor, FVector& HoverImpactPoint, bool& bWasHandled )
+void UVREditorUISystem::OnVRHoverUpdate( UViewportInteractor* Interactor, FVector& HoverImpactPoint, bool& bWasHandled )
 {
 	UVREditorMotionControllerInteractor* VREditorInteractor = Cast<UVREditorMotionControllerInteractor>( Interactor );
 	if( VREditorInteractor != nullptr )
@@ -886,7 +886,7 @@ void UVREditorUISystem::CreateUIs()
 		const bool bWithSceneComponent = false;
 		// @todo vreditor: Tweak
 		{
-			QuickMenuUI = GetOwner().SpawnTransientSceneActor< AVREditorFloatingUI >(TEXT("QuickMenu"), bWithSceneComponent);
+			QuickMenuUI = UViewportWorldInteraction::SpawnTransientSceneActor< AVREditorFloatingUI >(GetOwner().GetWorld(), TEXT("QuickMenu"), bWithSceneComponent);
 			QuickMenuUI->SetSlateWidget(*this, BuildQuickMenuWidget(), Resolution, 40.0f, AVREditorFloatingUI::EDockedTo::Nothing);
 			QuickMenuUI->ShowUI(false);
 			QuickMenuUI->SetRelativeOffset(FVector(-12.0f, 0.0f, 5.0f));
@@ -898,7 +898,7 @@ void UVREditorUISystem::CreateUIs()
 		}
 		// Create the radial UI
 		{
-			QuickRadialMenu = GetOwner().SpawnTransientSceneActor< AVREditorFloatingUI >(TEXT("QuickRadialmenu"), bWithSceneComponent);
+			QuickRadialMenu = UViewportWorldInteraction::SpawnTransientSceneActor< AVREditorFloatingUI >(GetOwner().GetWorld(), TEXT("QuickRadialmenu"), bWithSceneComponent);
 			QuickRadialMenu->SetSlateWidget(*this, BuildRadialMenuWidget(), DefaultResolution, 40.0f, AVREditorFloatingUI::EDockedTo::Nothing);
 			QuickRadialMenu->ShowUI(false);
 			QuickRadialMenu->SetActorEnableCollision(false);
@@ -961,7 +961,7 @@ void UVREditorUISystem::CreateUIs()
 			;
 
 			const bool bWithSceneComponent = false;
-			AVREditorFloatingUI* ContentBrowserUI = GetOwner().SpawnTransientSceneActor< AVREditorDockableWindow >( TEXT( "ContentBrowserUI" ), bWithSceneComponent );
+			AVREditorFloatingUI* ContentBrowserUI = UViewportWorldInteraction::SpawnTransientSceneActor< AVREditorDockableWindow >( GetOwner().GetWorld(), TEXT( "ContentBrowserUI" ), bWithSceneComponent );
 			ContentBrowserUI->SetSlateWidget( *this, WidgetToDraw, Resolution, VREd::ContentBrowserUISize->GetFloat(), AVREditorFloatingUI::EDockedTo::Nothing );
 			ContentBrowserUI->ShowUI( false );
 			FloatingUIs.Add( ContentBrowserUI );
@@ -985,7 +985,7 @@ void UVREditorUISystem::CreateUIs()
 			;
 
 			const bool bWithSceneComponent = false;
-			AVREditorFloatingUI* WorldOutlinerUI = GetOwner().SpawnTransientSceneActor< AVREditorDockableWindow >( TEXT( "WorldOutlinerUI" ), bWithSceneComponent );
+			AVREditorFloatingUI* WorldOutlinerUI = UViewportWorldInteraction::SpawnTransientSceneActor< AVREditorDockableWindow >( GetOwner().GetWorld(), TEXT( "WorldOutlinerUI" ), bWithSceneComponent );
 			WorldOutlinerUI->SetSlateWidget( *this, WidgetToDraw, DefaultResolution, VREd::EditorUISize->GetFloat(), AVREditorFloatingUI::EDockedTo::Nothing );
 			WorldOutlinerUI->ShowUI( false );
 			FloatingUIs.Add( WorldOutlinerUI );
@@ -1008,7 +1008,7 @@ void UVREditorUISystem::CreateUIs()
 			;
 
 			const bool bWithSceneComponent = false;
-			AVREditorFloatingUI* ActorDetailsUI = GetOwner().SpawnTransientSceneActor< AVREditorDockableWindow >( TEXT( "ActorDetailsUI" ), bWithSceneComponent );
+			AVREditorFloatingUI* ActorDetailsUI = UViewportWorldInteraction::SpawnTransientSceneActor< AVREditorDockableWindow >( GetOwner().GetWorld(), TEXT( "ActorDetailsUI" ), bWithSceneComponent );
 			ActorDetailsUI->SetSlateWidget( *this, WidgetToDraw, DefaultResolution, VREd::EditorUISize->GetFloat(), AVREditorFloatingUI::EDockedTo::Nothing );
 			ActorDetailsUI->ShowUI( false );
 			FloatingUIs.Add( ActorDetailsUI );
@@ -1030,7 +1030,7 @@ void UVREditorUISystem::CreateUIs()
 			;
 
 			const bool bWithSceneComponent = false;
-			AVREditorFloatingUI* ModesUI = GetOwner().SpawnTransientSceneActor< AVREditorDockableWindow >( TEXT( "ModesUI" ), bWithSceneComponent );
+			AVREditorFloatingUI* ModesUI = UViewportWorldInteraction::SpawnTransientSceneActor< AVREditorDockableWindow >( GetOwner().GetWorld(), TEXT( "ModesUI" ), bWithSceneComponent );
 			ModesUI->SetSlateWidget( *this, WidgetToDraw, DefaultResolution, VREd::EditorUISize->GetFloat(), AVREditorFloatingUI::EDockedTo::Nothing );
 			ModesUI->ShowUI( false );
 			FloatingUIs.Add( ModesUI );
@@ -1048,7 +1048,7 @@ void UVREditorUISystem::CreateUIs()
 				GetOwner().IsActuallyUsingVR() &&	// Don't show tutorial at startup when in "forced VR" mode, because you can't interact with UI using the mouse currently
 				(GetDefault<UEditorExperimentalSettings>()->bAlwaysShowVRTutorial); // Only show tutorial at startup if the user has this enabled
 
-			AVREditorFloatingUI* TutorialUI = GetOwner().SpawnTransientSceneActor< AVREditorDockableWindow >( TEXT( "TutorialUI" ), false );
+			AVREditorFloatingUI* TutorialUI = UViewportWorldInteraction::SpawnTransientSceneActor< AVREditorDockableWindow >( GetOwner().GetWorld(), TEXT( "TutorialUI" ), false );
 			TutorialUI->SetUMGWidget( *this, TutorialWidgetClass, FIntPoint( VREd::TutorialUIResolutionX->GetFloat(), VREd::TutorialUIResolutionY->GetFloat() ), VREd::TutorialUISize->GetFloat(), bShowAtStart ? AVREditorFloatingUI::EDockedTo::Room : AVREditorFloatingUI::EDockedTo::Nothing );
 
 			if( bShowAtStart )
@@ -1068,7 +1068,7 @@ void UVREditorUISystem::CreateUIs()
 			const FIntPoint Resolution = FIntPoint(VREd::AssetEditorUIResolutionX->GetInt(), VREd::AssetEditorUIResolutionY->GetInt());
 
 			const bool bWithSceneComponent = false;
-			AVREditorDockableWindow* TabManagerUI = GetOwner().SpawnTransientSceneActor< AVREditorDockableWindow >( TEXT( "AssetEditor" ), bWithSceneComponent );
+			AVREditorDockableWindow* TabManagerUI = UViewportWorldInteraction::SpawnTransientSceneActor< AVREditorDockableWindow >( GetOwner().GetWorld(), TEXT( "AssetEditor" ), bWithSceneComponent );
 			TabManagerUI->SetSlateWidget( *this, SNullWidget::NullWidget, Resolution, VREd::EditorUISize->GetFloat(), AVREditorFloatingUI::EDockedTo::Nothing );
 			TabManagerUI->ShowUI( false );
 
@@ -1115,7 +1115,7 @@ void UVREditorUISystem::CreateUIs()
 			;
 
 			const bool bWithSceneComponent = false;
-			AVREditorFloatingUI* WorldSettingsUI = GetOwner().SpawnTransientSceneActor< AVREditorDockableWindow >( TEXT( "WorldSettingsUI" ), bWithSceneComponent );
+			AVREditorFloatingUI* WorldSettingsUI = UViewportWorldInteraction::SpawnTransientSceneActor< AVREditorDockableWindow >( World, TEXT( "WorldSettingsUI" ), bWithSceneComponent );
 			WorldSettingsUI->SetSlateWidget( *this, WidgetToDraw, DefaultResolution, VREd::EditorUISize->GetFloat(), AVREditorFloatingUI::EDockedTo::Nothing );
 			WorldSettingsUI->ShowUI( false );
 			FloatingUIs.Add( WorldSettingsUI );
@@ -1138,7 +1138,7 @@ void UVREditorUISystem::CreateUIs()
 			;
 
 			const bool bWithSceneComponent = false;
-			AVREditorFloatingUI* SequencerUI = GetOwner().SpawnTransientSceneActor< AVREditorDockableWindow >(TEXT("SequencerUI"), bWithSceneComponent);
+			AVREditorFloatingUI* SequencerUI = UViewportWorldInteraction::SpawnTransientSceneActor< AVREditorDockableWindow >( GetOwner().GetWorld(), TEXT("SequencerUI"), bWithSceneComponent);
 			SequencerUI->SetSlateWidget(*this, WidgetToDraw, FIntPoint(VREd::ContentBrowserUIResolutionX->GetFloat(), VREd::ContentBrowserUIResolutionY->GetFloat()), VREd::ContentBrowserUISize->GetFloat(), AVREditorFloatingUI::EDockedTo::Nothing);
 			SequencerUI->ShowUI(false);
 			FloatingUIs.Add(SequencerUI);
@@ -1329,9 +1329,9 @@ void UVREditorUISystem::TryToSpawnRadialMenu( UVREditorInteractor* Interactor, c
 
 			bool bNeedsSpawn =
 				( QuickRadialMenu->bHidden || DockedToInteractor != Interactor ) &&
-				DraggingMode != EViewportInteractionDraggingMode::ActorsAtLaserImpact &&	// Don't show radial menu if the hand is busy dragging something around
-				DraggingMode != EViewportInteractionDraggingMode::ActorsFreely &&
-				DraggingMode != EViewportInteractionDraggingMode::ActorsWithGizmo &&
+				DraggingMode != EViewportInteractionDraggingMode::TransformablesAtLaserImpact &&	 // Don't show radial menu if the hand is busy dragging something around
+				DraggingMode != EViewportInteractionDraggingMode::TransformablesFreely &&
+				DraggingMode != EViewportInteractionDraggingMode::TransformablesWithGizmo &&
 				DraggingMode != EViewportInteractionDraggingMode::World &&
 				DraggingMode != EViewportInteractionDraggingMode::AssistingDrag &&
 				( InteractorDraggingUI == nullptr || InteractorDraggingUI != Interactor ) &&
@@ -1712,7 +1712,7 @@ void UVREditorUISystem::CreateVRColorPicker(const TSharedRef<SColorPicker>& Colo
 	if( ColorPickerUI == nullptr )
 	{
 		const bool bWithSceneComponent = false;
-		ColorPickerUI = GetOwner().SpawnTransientSceneActor< AVREditorDockableWindow >( TEXT( "ColorPickerUI" ), bWithSceneComponent );
+		ColorPickerUI = UViewportWorldInteraction::SpawnTransientSceneActor< AVREditorDockableWindow >( GetOwner().GetWorld(), TEXT( "ColorPickerUI" ), bWithSceneComponent );
 		FloatingUIs.Add( ColorPickerUI );
 		EditorUIPanels[ (int32)EEditorUIPanel::ColorPicker ] = ColorPickerUI;
 		bJustSpawned = true;
diff --git a/Engine/Source/Editor/VREditor/UI/VREditorUISystem.h b/Engine/Source/Editor/VREditor/UI/VREditorUISystem.h
--- a/Engine/Source/Editor/VREditor/UI/VREditorUISystem.h
+++ b/Engine/Source/Editor/VREditor/UI/VREditorUISystem.h
@@ -216,7 +216,7 @@ protected:
 		const struct FViewportActionKeyInput& Action, bool& bOutIsInputCaptured, bool& bWasHandled );
 
 	/** Called every frame to update hover state */
-	void OnVRHoverUpdate( class FEditorViewportClient& ViewportClient, UViewportInteractor* Interactor, FVector& HoverImpactPoint, bool& bWasHandled );
+	void OnVRHoverUpdate( UViewportInteractor* Interactor, FVector& HoverImpactPoint, bool& bWasHandled );
 
 	/** Testing Slate UI */
 	void CreateUIs();
diff --git a/Engine/Source/Editor/VREditor/VREditor.Build.cs b/Engine/Source/Editor/VREditor/VREditor.Build.cs
--- a/Engine/Source/Editor/VREditor/VREditor.Build.cs
+++ b/Engine/Source/Editor/VREditor/VREditor.Build.cs
@@ -7,7 +7,7 @@ namespace UnrealBuildTool.Rules
         public VREditor(TargetInfo Target)
 		{
 
-            PrivateDependencyModuleNames.AddRange(
+			PrivateDependencyModuleNames.AddRange(
                 new string[] {
                     "AppFramework",
 				    "Core",
@@ -21,13 +21,18 @@ namespace UnrealBuildTool.Rules
 					"UMG",
 					"LevelEditor",
 					"HeadMountedDisplay",
-                    "ViewportInteraction",
 					"Analytics",
                     "LevelSequence",
                     "Sequencer"
 				}
 			);
 
+			PublicDependencyModuleNames.AddRange(
+				new string[] {
+					"ViewportInteraction",
+				}
+			);
+
 			PrivateIncludePathModuleNames.AddRange(
 				new string[] {
 					"PlacementMode",
diff --git a/Engine/Source/Editor/VREditor/VREditorActions.cpp b/Engine/Source/Editor/VREditor/VREditorActions.cpp
--- a/Engine/Source/Editor/VREditor/VREditorActions.cpp
+++ b/Engine/Source/Editor/VREditor/VREditorActions.cpp
@@ -6,7 +6,6 @@
 #include "VREditorMode.h"
 #include "ViewportWorldInteraction.h"
 #include "VREditorInteractor.h"
-#include "EditorWorldManager.h"
 #include "VREditorFloatingUI.h"
 #include "VREditorTransformGizmo.h"
 #include "SLevelViewport.h"
diff --git a/Engine/Source/Editor/VREditor/VREditorInteractor.cpp b/Engine/Source/Editor/VREditor/VREditorInteractor.cpp
--- a/Engine/Source/Editor/VREditor/VREditorInteractor.cpp
+++ b/Engine/Source/Editor/VREditor/VREditorInteractor.cpp
@@ -47,7 +47,7 @@ void UVREditorInteractor::Shutdown()
 	for ( auto& KeyAndValue : HelpLabels )
 	{
 		AFloatingText* FloatingText = KeyAndValue.Value;
-		GetVRMode().DestroyTransientActor( FloatingText );
+		UViewportWorldInteraction::DestroyTransientActor( GetWorld(), FloatingText );
 	}
 
 	HelpLabels.Empty();	
@@ -75,9 +75,9 @@ FHitResult UVREditorInteractor::GetHitResultFromLaserPointer( TArray<AActor*>* O
 	OptionalListOfIgnoredActors->Add( GetVRMode().GetAvatarMeshActor() );
 	
 	// Ignore UI widgets too
-	if ( GetDraggingMode() == EViewportInteractionDraggingMode::ActorsAtLaserImpact )
+	if ( GetDraggingMode() == EViewportInteractionDraggingMode::TransformablesAtLaserImpact )
 	{
-		for( TActorIterator<AVREditorFloatingUI> UIActorIt( WorldInteraction->GetViewportWorld() ); UIActorIt; ++UIActorIt )
+		for( TActorIterator<AVREditorFloatingUI> UIActorIt( WorldInteraction->GetWorld() ); UIActorIt; ++UIActorIt )
 		{
 			OptionalListOfIgnoredActors->Add( *UIActorIt );
 		}
@@ -102,12 +102,11 @@ void UVREditorInteractor::ResetHoverState( const float DeltaTime )
 	bIsHoveringOverUI = false;
 }	
 
-void UVREditorInteractor::OnStartDragging( UActorComponent* ClickedComponent, const FVector& HitLocation, const bool bIsPlacingActors )
+void UVREditorInteractor::OnStartDragging( const FVector& HitLocation, const bool bIsPlacingNewObjects )
 {
-	// If the user is holding down the modifier key, go ahead and duplicate the selection first.  This needs to
-	// happen before we create our transformables, because duplicating actors will change which actors are
-	// selected!  Don't do this if we're placing objects right now though.
-	if ( bIsModifierPressed && !bIsPlacingActors )
+	// If the user is holding down the modifier key, go ahead and duplicate the selection first.  Don't do this if we're
+	// placing objects right now though.
+	if ( bIsModifierPressed && !bIsPlacingNewObjects )
 	{
 		WorldInteraction->Duplicate();
 	}
@@ -208,7 +207,7 @@ float UVREditorInteractor::GetSelectAndMoveTriggerValue() const
 	return SelectAndMoveTriggerValue;
 }
 
-bool UVREditorInteractor::GetIsLaserBlocked()
+bool UVREditorInteractor::GetIsLaserBlocked() const
 {
 	return bHasUIInFront;
 }
diff --git a/Engine/Source/Editor/VREditor/VREditorMode.cpp b/Engine/Source/Editor/VREditor/VREditorMode.cpp
--- a/Engine/Source/Editor/VREditor/VREditorMode.cpp
+++ b/Engine/Source/Editor/VREditor/VREditorMode.cpp
@@ -29,10 +29,10 @@
 #include "Interfaces/IAnalyticsProvider.h"
 
 #include "IViewportInteractionModule.h"
-#include "MouseCursorInteractor.h"
 #include "VREditorMotionControllerInteractor.h"
 
 #include "ViewportWorldInteractionManager.h"
+#include "EditorWorldExtension.h"
 
 #include "ISequencer.h"
 
@@ -47,25 +47,23 @@ namespace VREd
 
 // @todo vreditor: Hacky that we have to import these this way. (Plugin has them in a .cpp, not exported)
 
-UVREditorMode::UVREditorMode( const FObjectInitializer& ObjectInitializer ) : 
-	Super( ObjectInitializer ),
+UVREditorMode::UVREditorMode() : 
+	Super(),
 	bWantsToExitMode( false ),
 	ExitType( EVREditorExitType::Normal ),
 	bIsFullyInitialized( false ),
 	AppTimeModeEntered( FTimespan::Zero() ),
 	AvatarActor( nullptr ),
-    	FlashlightComponent( nullptr ),
+   	FlashlightComponent( nullptr ),
 	bIsFlashlightOn( false ),
 	MotionControllerID( 0 ),	// @todo vreditor minor: We only support a single controller, and we assume the first controller are the motion controls
 	UISystem( nullptr ),
 	TeleporterSystem( nullptr ),
 	AutoScalerSystem( nullptr ),
 	WorldInteraction( nullptr ),
-	MouseCursorInteractor( nullptr ),
 	LeftHandInteractor( nullptr ),
 	RightHandInteractor( nullptr ),
 	bFirstTick( true ),
-	bWasInWorldSpaceBeforeScaleMode( false ),
 	bIsActive( false )
 {
 }
@@ -76,7 +74,8 @@ UVREditorMode::~UVREditorMode()
 	Shutdown();
 }
 
-void UVREditorMode::Init( UViewportWorldInteraction* InViewportWorldInteraction )
+
+void UVREditorMode::Init()
 {
 	// @todo vreditor urgent: Turn on global editor hacks for VR Editor mode
 	GEnableVREditorHacks = true;
@@ -112,7 +111,11 @@ void UVREditorMode::Init( UViewportWorldInteraction* InViewportWorldInteraction
 		Colors[ (int32)EColors::UICloseButtonHoverColor ] = FLinearColor( 1.0f, 1.0f, 1.0f, 1.0f );
 	}
 
-	WorldInteraction = InViewportWorldInteraction;
+	{
+		TSharedPtr<FEditorWorldExtensionCollection> Collection = GEditor->GetEditorWorldExtensionsManager()->GetEditorWorldExtensions( GetWorld() );
+		WorldInteraction = Cast<UViewportWorldInteraction>( Collection->FindExtension( UViewportWorldInteraction::StaticClass() ) );
+		check( WorldInteraction != nullptr );
+	}
 
 	bIsFullyInitialized = true;
 }
@@ -127,7 +130,6 @@ void UVREditorMode::Shutdown()
 	TeleporterSystem = nullptr;
 	AutoScalerSystem = nullptr;
 	WorldInteraction = nullptr;
-	MouseCursorInteractor = nullptr;
 	LeftHandInteractor = nullptr;
 	RightHandInteractor = nullptr;
 
@@ -141,9 +143,8 @@ void UVREditorMode::Enter(const bool bReenteringVREditing)
 	ExitType = EVREditorExitType::Normal;
 
 	{
-		IViewportWorldInteractionManager& ViewportWorldInteraction = IViewportInteractionModule::Get().GetWorldInteractionManager();
-		ViewportWorldInteraction.OnPreWorldInteractionTick().AddUObject( this, &UVREditorMode::PreTick );
-		ViewportWorldInteraction.OnPostWorldInteractionTick().AddUObject( this, &UVREditorMode::Tick );
+		WorldInteraction->OnPreWorldInteractionTick().AddUObject( this, &UVREditorMode::PreTick );
+		WorldInteraction->OnPostWorldInteractionTick().AddUObject( this, &UVREditorMode::PostTick );
 	}
 
 	// @todo vreditor: We need to make sure the user can never switch to orthographic mode, or activate settings that
@@ -326,6 +327,10 @@ void UVREditorMode::Enter(const bool bReenteringVREditing)
 			// like Content Browser drag and drop.
 			SavedEditorState.DragTriggerDistance = FSlateApplication::Get().GetDragTriggerDistance();
 			FSlateApplication::Get().SetDragTriggerDistance( 100.0f );	// @todo vreditor tweak
+
+			// When actually in VR, make sure the transform gizmo is big!
+			SavedEditorState.TransformGizmoScale = WorldInteraction->GetTransformGizmoScale();
+			WorldInteraction->SetTransformGizmoScale( 1.0f );
 		}
 	}
 
@@ -333,32 +338,35 @@ void UVREditorMode::Enter(const bool bReenteringVREditing)
 	{
 		// Setup world interaction
 		TSharedPtr<FEditorViewportClient> ViewportClient = VREditorLevelViewportWeakPtr.Pin()->GetViewportClient();
-		WorldInteraction->SetViewport( ViewportClient );
-		WorldInteraction->Activate( true );
 
-		WorldInteraction->OnHandleKeyInput().AddUObject( this, &UVREditorMode::InputKey );
+		// We're not expecting anyone to have already set a default viewport client
+		check( WorldInteraction->GetDefaultOptionalViewportClient() == nullptr );
+		WorldInteraction->SetDefaultOptionalViewportClient( ViewportClient );
+
+		// We need input preprocessing for VR so that we can receive motion controller input without any viewports having 
+		// to be focused.  This is mainly because Slate UI injected into the 3D world can cause focus to be lost unexpectedly,
+		// but we need the user to still be able to interact with UI.
+		WorldInteraction->SetUseInputProcessor( true );
 
 		// Motion controllers
-		if(bActuallyUsingVR)
 		{
-			LeftHandInteractor = NewObject<UVREditorMotionControllerInteractor>( WorldInteraction );
+			LeftHandInteractor = NewObject<UVREditorMotionControllerInteractor>();
 			LeftHandInteractor->SetControllerHandSide( EControllerHand::Left );
 			LeftHandInteractor->Init( this );
 			WorldInteraction->AddInteractor( LeftHandInteractor );
 
-			RightHandInteractor = NewObject<UVREditorMotionControllerInteractor>( WorldInteraction );
+			RightHandInteractor = NewObject<UVREditorMotionControllerInteractor>();
 			RightHandInteractor->SetControllerHandSide( EControllerHand::Right );
 			RightHandInteractor->Init( this );
 			WorldInteraction->AddInteractor( RightHandInteractor );
 
 			WorldInteraction->PairInteractors( LeftHandInteractor, RightHandInteractor );
 		}
-		else
+
+		if( !bActuallyUsingVR )
 		{
-			// Register an interactor for the mouse cursor
-			MouseCursorInteractor = NewObject<UMouseCursorInteractor>( WorldInteraction );
-			MouseCursorInteractor->Init();
-			WorldInteraction->AddInteractor( MouseCursorInteractor );
+			// When not actually using VR devices, we want an interactor for our mouse cursor!
+			WorldInteraction->AddMouseCursorInteractor();
 		}
 
 		// Setup the UI system
@@ -395,7 +403,7 @@ void UVREditorMode::Exit(const bool bHMDShouldExitStereo)
 	{
 		//Destroy the avatar
 		{
-			DestroyTransientActor( AvatarActor );
+			UViewportWorldInteraction::DestroyTransientActor( GetWorld(), AvatarActor );
 			AvatarActor = nullptr;
 			FlashlightComponent = nullptr;
 		}
@@ -411,6 +419,9 @@ void UVREditorMode::Exit(const bool bHMDShouldExitStereo)
 			{
 				// Restore Slate drag trigger distance
 				FSlateApplication::Get().SetDragTriggerDistance( SavedEditorState.DragTriggerDistance );
+
+				// Restore gizmo size
+				WorldInteraction->SetTransformGizmoScale( SavedEditorState.TransformGizmoScale );
 			}
 
 			TSharedPtr<SLevelViewport> VREditorLevelViewport( VREditorLevelViewportWeakPtr.Pin() );
@@ -488,7 +499,7 @@ void UVREditorMode::Exit(const bool bHMDShouldExitStereo)
 	}
 	//Destroy the avatar
 	{
-		DestroyTransientActor( AvatarActor );
+		UViewportWorldInteraction::DestroyTransientActor( GetWorld(), AvatarActor );
 		AvatarActor = nullptr;
 		FlashlightComponent = nullptr;
 	}
@@ -503,7 +514,6 @@ void UVREditorMode::Exit(const bool bHMDShouldExitStereo)
 
 	if( VRWorldInteractionExtension != nullptr )
 	{
-		VRWorldInteractionExtension->Shutdown();
 		VRWorldInteractionExtension->MarkPendingKill();
 		VRWorldInteractionExtension = nullptr;
 	}
@@ -522,26 +532,28 @@ void UVREditorMode::Exit(const bool bHMDShouldExitStereo)
 		AutoScalerSystem = nullptr;
 	}
 
-	{
-		IViewportWorldInteractionManager& ViewportWorldInteraction = IViewportInteractionModule::Get().GetWorldInteractionManager();
-		ViewportWorldInteraction.OnPreWorldInteractionTick().RemoveAll( this );
-		ViewportWorldInteraction.OnPostWorldInteractionTick().RemoveAll( this );
-	}
-
 	if( WorldInteraction != nullptr )
 	{
 		WorldInteraction->OnHandleKeyInput().RemoveAll( this );
+		WorldInteraction->OnPreWorldInteractionTick().RemoveAll( this );
+		WorldInteraction->OnPostWorldInteractionTick().RemoveAll( this );
 
 		WorldInteraction->RemoveInteractor( LeftHandInteractor );
 		LeftHandInteractor->MarkPendingKill();
 		LeftHandInteractor->Shutdown();
 		LeftHandInteractor = nullptr;
+
 		WorldInteraction->RemoveInteractor( RightHandInteractor );
 		RightHandInteractor->Shutdown();
 		RightHandInteractor->MarkPendingKill();
 		RightHandInteractor = nullptr;
+		
+		if( !bActuallyUsingVR )
+		{
+			WorldInteraction->ReleaseMouseCursorInteractor();
+		}
 
-		WorldInteraction->Activate( false );
+		WorldInteraction->SetDefaultOptionalViewportClient( nullptr );
 	}
 
 	GEditor->OnEditorClose().RemoveAll( this );
@@ -573,7 +585,7 @@ void UVREditorMode::SpawnAvatarMeshActor()
 	{
 		{
 			const bool bWithSceneComponent = true;
-			AvatarActor = SpawnTransientSceneActor<AVREditorAvatarActor>( TEXT( "AvatarActor" ), bWithSceneComponent );
+			AvatarActor = UViewportWorldInteraction::SpawnTransientSceneActor<AVREditorAvatarActor>( GetWorld(), TEXT( "AvatarActor" ), bWithSceneComponent );
 			AvatarActor->Init( this );
 		}
 
@@ -619,7 +631,7 @@ void UVREditorMode::PreTick( const float DeltaTime )
 
 }
 
-void UVREditorMode::Tick( float DeltaTime )
+void UVREditorMode::PostTick( float DeltaTime )
 {
 	if( !bIsFullyInitialized || !bIsActive || bWantsToExitMode || !VREditorLevelViewportWeakPtr.IsValid() )
 	{
@@ -650,59 +662,19 @@ void UVREditorMode::Tick( float DeltaTime )
 	bFirstTick = false;
 }
 
-
-void UVREditorMode::InputKey(const FEditorViewportClient& InViewportClient, const FKey InKey, const EInputEvent InEvent, bool& bOutWasHandled)
+bool UVREditorMode::InputKey( FEditorViewportClient* InViewportClient, FViewport* Viewport, FKey Key, EInputEvent Event )
 {
-	if(InKey == EKeys::Escape)
+	bool bHandled = false;
+	if( Key == EKeys::Escape )
 	{
 		// User hit escape, so bail out of VR mode
 		StartExitingVRMode();
+		bHandled = true;
 	}
+
+	return bHandled;
 }
 
-//bool UVREditorMode::InputKey(FEditorViewportClient* ViewportClient, FViewport* Viewport, FKey Key, EInputEvent Event)
-//{
-//	// Only if this is our VR viewport. Remember, editor modes get a chance to tick and receive input for each active viewport.
-//	if ( ViewportClient != GetLevelViewportPossessedForVR().GetViewportClient().Get() )
-//	{
-//		return InputKey(GetLevelViewportPossessedForVR().GetViewportClient().Get(), Viewport, Key, Event);
-//	}
-//
-//	if ( Key == EKeys::Escape )
-//	{
-//		// User hit escape, so bail out of VR mode
-//		StartExitingVRMode();
-//	}
-//	else if( Key.IsMouseButton() )	// Input preprocessor cannot handle mouse buttons, so we'll route those the normal way
-//	{
-//		return WorldInteraction->HandleInputKey( Key, Event );
-//	}
-//
-//	return FEdMode::InputKey(ViewportClient, Viewport, Key, Event);
-//}
-//
-//
-//bool UVREditorMode::InputAxis(FEditorViewportClient* ViewportClient, FViewport* Viewport, int32 ControllerId, FKey Key, float Delta, float DeltaTime)
-//{
-//	// Only if this is our VR viewport. Remember, editor modes get a chance to tick and receive input for each active viewport.
-//	if ( ViewportClient != GetLevelViewportPossessedForVR().GetViewportClient().Get() )
-//	{
-//		return InputAxis(GetLevelViewportPossessedForVR().GetViewportClient().Get(), Viewport, ControllerId, Key, Delta, DeltaTime);
-//	}
-//
-//	return FEdMode::InputAxis(ViewportClient, Viewport, ControllerId, Key, Delta, DeltaTime);
-//}
-
-//void UVREditorMode::Render( const FSceneView* SceneView, FViewport* Viewport, FPrimitiveDrawInterface* PDI )
-//{
-//	//StopOldHapticEffects(); //@todo vreditor
-//
-//	if( bIsFullyInitialized )
-//	{
-//		// Let our subsystems render, too
-//		UISystem->Render( SceneView, Viewport, PDI );
-//	}
-//}
 
 /************************************************************************/
 /* IVREditorMode interface                                              */
@@ -715,7 +687,17 @@ AActor* UVREditorMode::GetAvatarMeshActor()
 
 UWorld* UVREditorMode::GetWorld() const
 {
-	return WorldInteraction->GetViewportWorld();
+	UWorld* ResultedWorld = nullptr;
+	if( WorldInteraction )
+	{
+		ResultedWorld = WorldInteraction->GetWorld();
+	}
+	else
+	{
+		ResultedWorld = UEditorWorldExtension::GetWorld();
+	}
+
+	return ResultedWorld;
 }
 
 FTransform UVREditorMode::GetRoomTransform() const
@@ -758,44 +740,6 @@ bool UVREditorMode::IsActive() const
 	return bIsActive;
 }
 
-AActor* UVREditorMode::SpawnTransientSceneActor(TSubclassOf<AActor> ActorClass, const FString& ActorName, const bool bWithSceneComponent) const
-{
-	const bool bWasWorldPackageDirty = GetWorld()->GetOutermost()->IsDirty();
-
-	// @todo vreditor: Needs respawn if world changes (map load, etc.)  Will that always restart the editor mode anyway?
-	FActorSpawnParameters ActorSpawnParameters;
-	ActorSpawnParameters.Name = MakeUniqueObjectName( GetWorld(), ActorClass, *ActorName );	// @todo vreditor: Without this, SpawnActor() can return us an existing PendingKill actor of the same name!  WTF?
-	ActorSpawnParameters.SpawnCollisionHandlingOverride = ESpawnActorCollisionHandlingMethod::AlwaysSpawn;
-	ActorSpawnParameters.ObjectFlags = EObjectFlags::RF_Transient;
-
-	check( ActorClass != nullptr );
-	AActor* NewActor = GetWorld()->SpawnActor< AActor >( ActorClass, ActorSpawnParameters );
-	NewActor->SetActorLabel( ActorName );
-
-	if( bWithSceneComponent )
-	{
-		// Give the new actor a root scene component, so we can attach multiple sibling components to it
-		USceneComponent* SceneComponent = NewObject<USceneComponent>( NewActor );
-		NewActor->AddOwnedComponent( SceneComponent );
-		NewActor->SetRootComponent( SceneComponent );
-		SceneComponent->RegisterComponent();
-	}
-
-	// Don't dirty the level file after spawning a transient actor
-	if( !bWasWorldPackageDirty )
-	{
-		GetWorld()->GetOutermost()->SetDirtyFlag( false );
-	}
-
-	return NewActor;
-}
-
-
-void UVREditorMode::DestroyTransientActor( AActor* Actor ) const
-{
-	WorldInteraction->DestroyTransientActor( Actor );
-}
-
 bool UVREditorMode::IsShowingRadialMenu(const UVREditorInteractor* Interactor) const
 {
 	return UISystem->IsShowingRadialMenu(Interactor);
@@ -823,7 +767,7 @@ void UVREditorMode::CleanUpActorsBeforeMapChangeOrSimulate()
 	if ( WorldInteraction != nullptr )
 	{
 		// NOTE: This will be called even when this mode is not currently active!
-		DestroyTransientActor( AvatarActor );
+		UViewportWorldInteraction::DestroyTransientActor( GetWorld(), AvatarActor );
 		AvatarActor = nullptr;
 		FlashlightComponent = nullptr;
 
@@ -873,27 +817,6 @@ void UVREditorMode::CycleTransformGizmoHandleType()
 		NewGizmoType = EGizmoHandleTypes::All;
 	}
 
-	// Set coordinate system to local if the next gizmo will be for non-uniform scaling 
-	if ( NewGizmoType == EGizmoHandleTypes::Scale )
-	{
-		const ECoordSystem CurrentCoordSystem = WorldInteraction->GetTransformGizmoCoordinateSpace();
-		if ( CurrentCoordSystem == COORD_World )
-		{
-			GLevelEditorModeTools().SetCoordSystem( COORD_Local );
-			// Remember if coordinate system was in world space before scaling
-			bWasInWorldSpaceBeforeScaleMode = true;
-		}
-		else if ( CurrentCoordSystem == COORD_Local )
-		{
-			bWasInWorldSpaceBeforeScaleMode = false;
-		}
-	} 
-	else if ( WorldInteraction->GetCurrentGizmoType() == EGizmoHandleTypes::Scale && bWasInWorldSpaceBeforeScaleMode )
-	{
-		// Set the coordinate system to world space if the coordinate system was world before scaling
-		WorldInteraction->SetTransformGizmoCoordinateSpace( COORD_World );
-	}
-	
 	WorldInteraction->SetGizmoHandleType( NewGizmoType );
 }
 
diff --git a/Engine/Source/Editor/VREditor/VREditorModeManager.cpp b/Engine/Source/Editor/VREditor/VREditorModeManager.cpp
--- a/Engine/Source/Editor/VREditor/VREditorModeManager.cpp
+++ b/Engine/Source/Editor/VREditor/VREditorModeManager.cpp
@@ -12,7 +12,8 @@
 #include "EngineGlobals.h"
 #include "LevelEditor.h"
 #include "IHeadMountedDisplay.h"
-#include "EditorWorldManager.h"
+#include "EditorWorldExtension.h"
+#include "ViewportWorldInteraction.h"
 
 FVREditorModeManager::FVREditorModeManager() :
 	CurrentVREditorMode( nullptr ),
@@ -151,6 +152,11 @@ bool FVREditorModeManager::IsVREditorAvailable() const
 }
 
 
+UVREditorMode* FVREditorModeManager::GetCurrentVREditorMode()
+{
+	return CurrentVREditorMode;
+}
+
 void FVREditorModeManager::AddReferencedObjects( FReferenceCollector& Collector )
 {
 	Collector.AddReferencedObject( CurrentVREditorMode );
@@ -168,8 +174,13 @@ void FVREditorModeManager::StartVREditorMode( const bool bForceWithoutHMD )
 		bPlayStartedFromVREditor = false;
 	}
 
-	TSharedPtr<FEditorWorldWrapper> EditorWorld = GEditor->GetEditorWorldManager()->GetEditorWorldWrapper( GWorld );
-	UVREditorMode* ModeFromWorld = EditorWorld->GetVREditorMode();
+	UVREditorMode* VRMode = nullptr;
+	{
+		TSharedPtr<FEditorWorldExtensionCollection> Collection = GEditor->GetEditorWorldExtensionsManager()->GetEditorWorldExtensions( GWorld );
+		Collection->AddExtension( UViewportWorldInteraction::StaticClass() );
+		VRMode = Cast<UVREditorMode>( Collection->AddExtension( UVREditorMode::StaticClass() ) );
+		check( VRMode != nullptr );
+	}
 
 	// Tell the level editor we want to be notified when selection changes
 	{
@@ -177,7 +188,7 @@ void FVREditorModeManager::StartVREditorMode( const bool bForceWithoutHMD )
 		LevelEditor.OnMapChanged().AddRaw( this, &FVREditorModeManager::OnMapChanged );
 	}
 	
-	CurrentVREditorMode = ModeFromWorld;
+	CurrentVREditorMode = VRMode;
 	CurrentVREditorMode->SetActuallyUsingVR( !bForceWithoutHMD );
 	CurrentVREditorMode->Enter(bReenteringVREditing);
 }
@@ -199,6 +210,10 @@ void FVREditorModeManager::CloseVREditor( const bool bHMDShouldExitStereo )
 		PreviousVREditorMode = CurrentVREditorMode;
 		CurrentVREditorMode = nullptr;
 	}
+
+	TSharedPtr<FEditorWorldExtensionCollection> Collection = GEditor->GetEditorWorldExtensionsManager()->GetEditorWorldExtensions( GWorld );
+	Collection->RemoveExtension( Collection->FindExtension( UVREditorMode::StaticClass() ) );
+	Collection->RemoveExtension( Collection->FindExtension( UViewportWorldInteraction::StaticClass() ) );
 }
 
 void FVREditorModeManager::SetDirectWorldToMeters( const float NewWorldToMeters )
diff --git a/Engine/Source/Editor/VREditor/VREditorModeManager.h b/Engine/Source/Editor/VREditor/VREditorModeManager.h
--- a/Engine/Source/Editor/VREditor/VREditorModeManager.h
+++ b/Engine/Source/Editor/VREditor/VREditorModeManager.h
@@ -18,7 +18,10 @@ class FVREditorModeManager : public FGCObject
 
 public:
 
+	/** Default constructor */
 	FVREditorModeManager();
+
+	/** Default destructor */
 	~FVREditorModeManager();
 
 	/** Ticks to detect entering and closing the VR Editor */
@@ -33,7 +36,12 @@ public:
 	/** If the VR Editor is currently available */
 	bool IsVREditorAvailable() const;
 
+	/** Gets the current VR Editor mode that was enabled */
+	UVREditorMode* GetCurrentVREditorMode();
+
+	// FGCObject
 	virtual void AddReferencedObjects( FReferenceCollector& Collector );
+	// End FGCObject
 
 private:
 
diff --git a/Engine/Source/Editor/VREditor/VREditorModule.cpp b/Engine/Source/Editor/VREditor/VREditorModule.cpp
--- a/Engine/Source/Editor/VREditor/VREditorModule.cpp
+++ b/Engine/Source/Editor/VREditor/VREditorModule.cpp
@@ -8,6 +8,7 @@
 #include "TickableEditorObject.h"
 #include "VREditorModeManager.h"
 #include "VREditorStyle.h"
+#include "VREditorMode.h"
 #include "Kismet/HeadMountedDisplayFunctionLibrary.h"	// For EHMDWornState::Type
 
 class FVREditorModule : public IVREditorModule, public FTickableEditorObject
@@ -31,6 +32,7 @@ public:
 	virtual bool IsVREditorAvailable() const override;
 	virtual void EnableVREditor( const bool bEnable, const bool bForceWithoutHMD ) override;
 	virtual bool IsVREditorModeActive() override;
+	virtual UVREditorMode* GetVRMode() override;
 
 	// FTickableEditorObject overrides
 	virtual void Tick( float DeltaTime ) override;
@@ -100,6 +102,11 @@ bool FVREditorModule::IsVREditorModeActive()
 	return ModeManager.IsVREditorActive();
 }
 
+UVREditorMode* FVREditorModule::GetVRMode()
+{
+	return ModeManager.GetCurrentVREditorMode();
+}
+
 void FVREditorModule::ToggleForceVRMode()
 {
 	const bool bForceWithoutHMD = true;
diff --git a/Engine/Source/Editor/VREditor/VREditorMotionControllerInteractor.cpp b/Engine/Source/Editor/VREditor/VREditorMotionControllerInteractor.cpp
--- a/Engine/Source/Editor/VREditor/VREditorMotionControllerInteractor.cpp
+++ b/Engine/Source/Editor/VREditor/VREditorMotionControllerInteractor.cpp
@@ -19,7 +19,6 @@
 #include "UObject/Package.h"
 #include "Engine/Engine.h"
 #include "EngineGlobals.h"
-
 #include "IMotionController.h"
 #include "IHeadMountedDisplay.h"
 #include "MotionControllerComponent.h"
@@ -463,10 +462,10 @@ void UVREditorMotionControllerInteractor::Tick( const float DeltaTime )
 			ResultColorType = UVREditorMode::EColors::WorldDraggingColor_OneHand;
 			//			}
 		}
-		else if ( DraggingMode == EViewportInteractionDraggingMode::ActorsAtLaserImpact ||
+		else if ( DraggingMode == EViewportInteractionDraggingMode::TransformablesAtLaserImpact ||
 			DraggingMode == EViewportInteractionDraggingMode::Material ||
-			DraggingMode == EViewportInteractionDraggingMode::ActorsFreely ||
-			DraggingMode == EViewportInteractionDraggingMode::ActorsWithGizmo ||
+			DraggingMode == EViewportInteractionDraggingMode::TransformablesFreely ||
+			DraggingMode == EViewportInteractionDraggingMode::TransformablesWithGizmo ||
 			DraggingMode == EViewportInteractionDraggingMode::AssistingDrag ||
 			DraggingMode == EViewportInteractionDraggingMode::Interactable ||
 			( GetVRMode().GetUISystem().IsInteractorDraggingDockUI( this ) && GetVRMode().GetUISystem().IsDraggingDockUI() ) )
@@ -551,7 +550,7 @@ EHMDDeviceType::Type UVREditorMotionControllerInteractor::GetHMDDeviceType() con
 }
 
 
-void UVREditorMotionControllerInteractor::HandleInputKey( FViewportActionKeyInput& Action, const FKey Key, const EInputEvent Event, bool& bOutWasHandled )
+void UVREditorMotionControllerInteractor::HandleInputKey( FEditorViewportClient& ViewportClient, FViewportActionKeyInput& Action, const FKey Key, const EInputEvent Event, bool& bOutWasHandled )
 {
 	if ( !bOutWasHandled )
 	{
@@ -567,7 +566,7 @@ void UVREditorMotionControllerInteractor::HandleInputKey( FViewportActionKeyInpu
 				const float TimeSinceLightlyPressed = ( float ) ( FPlatformTime::Seconds() - GetRealTimeTriggerWasLightlyPressed() );
 				if ( !IsClickingOnUI() &&	// @todo vreditor: UI clicks happen with light presses; we never want to convert to a hard press!
 					GetDraggingMode() != EViewportInteractionDraggingMode::Material &&	// @todo vreditor: Material dragging happens with light presses, don't convert to a hard press!
-					GetDraggingMode() != EViewportInteractionDraggingMode::ActorsAtLaserImpact &&	// @todo vreditor: Actor dragging happens with light presses, don't convert to a hard press!
+					GetDraggingMode() != EViewportInteractionDraggingMode::TransformablesAtLaserImpact &&	// @todo vreditor: Actor dragging happens with light presses, don't convert to a hard press!
 					AllowTriggerFullPress() &&
 					( !AllowTriggerLightPressLocking() || TimeSinceLightlyPressed < VREd::TriggerLightlyPressedLockTime->GetFloat() ) )
 				{
@@ -577,7 +576,7 @@ void UVREditorMotionControllerInteractor::HandleInputKey( FViewportActionKeyInpu
 					// Synthesize an input key for releasing the light press
 					// NOTE: Intentionally re-entrant call here.
 					const EInputEvent InputEvent = IE_Released;
-					const bool bWasLightReleaseHandled = UViewportInteractor::HandleInputKey( GetControllerSide() == EControllerHand::Left ? VREditorKeyNames::MotionController_Left_LightlyPressedTriggerAxis : VREditorKeyNames::MotionController_Right_LightlyPressedTriggerAxis, InputEvent );
+					const bool bWasLightReleaseHandled = UViewportInteractor::HandleInputKey( ViewportClient, GetControllerSide() == EControllerHand::Left ? VREditorKeyNames::MotionController_Left_LightlyPressedTriggerAxis : VREditorKeyNames::MotionController_Right_LightlyPressedTriggerAxis, InputEvent );
 				}
 				else
 				{
@@ -639,20 +638,20 @@ void UVREditorMotionControllerInteractor::HandleInputKey( FViewportActionKeyInpu
 		if( Event == IE_Pressed )
 		{ 
 			// Start dragging at laser impact when already dragging actors freely
-			if( DraggingMode == EViewportInteractionDraggingMode::ActorsFreely )
+			if( DraggingMode == EViewportInteractionDraggingMode::TransformablesFreely )
 			{
 				FVector PlaceAt = GetHoverLocation();
 				const bool bIsPlacingActors = true;
-				const FViewportActionKeyInput FakeSelectAndMoveAction( ViewportWorldActionTypes::SelectAndMove_LightlyPressed );
-				WorldInteraction->StartDraggingActors( this, FakeSelectAndMoveAction, WorldInteraction->GetTransformGizmoActor()->GetRootComponent(), PlaceAt, bIsPlacingActors );
+				const bool bStartTransaction = true;
+				WorldInteraction->StartDragging( this, WorldInteraction->GetTransformGizmoActor()->GetRootComponent(), PlaceAt, bIsPlacingActors, bStartTransaction );
 			}
 		}
 		else if( Event == IE_Released )
 		{
 			// Disable dragging at laser impact when releasing
-			if( DraggingMode == EViewportInteractionDraggingMode::ActorsAtLaserImpact )
+			if( DraggingMode == EViewportInteractionDraggingMode::TransformablesAtLaserImpact )
 			{
-				SetDraggingMode( EViewportInteractionDraggingMode::ActorsFreely );
+				SetDraggingMode( EViewportInteractionDraggingMode::TransformablesFreely );
 			}
 		}
 	}
@@ -660,7 +659,7 @@ void UVREditorMotionControllerInteractor::HandleInputKey( FViewportActionKeyInpu
 	ApplyButtonPressColors( Action );
 }
 
-void UVREditorMotionControllerInteractor::HandleInputAxis( FViewportActionKeyInput& Action, const FKey Key, const float Delta, const float DeltaTime, bool& bOutWasHandled )
+void UVREditorMotionControllerInteractor::HandleInputAxis( FEditorViewportClient& ViewportClient, FViewportActionKeyInput& Action, const FKey Key, const float Delta, const float DeltaTime, bool& bOutWasHandled )
 {
 	if ( Action.ActionType == TriggerAxis )
 	{
@@ -674,7 +673,7 @@ void UVREditorMotionControllerInteractor::HandleInputAxis( FViewportActionKeyInp
 
 			FViewportActionKeyInput* OptionalAction = GetActionWithName( ViewportWorldActionTypes::SelectAndMove );
 			const bool bIsFullPressAlreadyCapturing = OptionalAction && OptionalAction->bIsInputCaptured;
-				
+
 			if ( !bIsFullPressAlreadyCapturing &&		// Don't fire if we're already capturing input for a full press
 				!bIsTriggerAtLeastLightlyPressed &&	// Don't fire unless we are already pressed
 				bHasTriggerBeenReleasedSinceLastPress &&	// Only if we've been fully released since the last time we fired
@@ -688,7 +687,7 @@ void UVREditorMotionControllerInteractor::HandleInputAxis( FViewportActionKeyInp
 
 				// Synthesize an input key for this light press
 				const EInputEvent InputEvent = IE_Pressed;
-				const bool bWasLightPressHandled = UViewportInteractor::HandleInputKey( ControllerHandSide == EControllerHand::Left ? MotionController_Left_LightlyPressedTriggerAxis : MotionController_Right_LightlyPressedTriggerAxis, InputEvent );
+				const bool bWasLightPressHandled = UViewportInteractor::HandleInputKey( ViewportClient, ControllerHandSide == EControllerHand::Left ? MotionController_Left_LightlyPressedTriggerAxis : MotionController_Right_LightlyPressedTriggerAxis, InputEvent );
 			}
 			else if ( bIsTriggerAtLeastLightlyPressed && Delta < TriggerLightlyPressedThreshold )
 			{
@@ -696,7 +695,7 @@ void UVREditorMotionControllerInteractor::HandleInputAxis( FViewportActionKeyInp
 
 				// Synthesize an input key for this light press
 				const EInputEvent InputEvent = IE_Released;
-				const bool bWasLightReleaseHandled = UViewportInteractor::HandleInputKey( ControllerHandSide == EControllerHand::Left ? MotionController_Left_LightlyPressedTriggerAxis : MotionController_Right_LightlyPressedTriggerAxis, InputEvent );
+				const bool bWasLightReleaseHandled = UViewportInteractor::HandleInputKey( ViewportClient, ControllerHandSide == EControllerHand::Left ? MotionController_Left_LightlyPressedTriggerAxis : MotionController_Right_LightlyPressedTriggerAxis, InputEvent );
 			}
 		}
 
@@ -715,7 +714,7 @@ void UVREditorMotionControllerInteractor::HandleInputAxis( FViewportActionKeyInp
 
 				// Synthesize an input key for this full press
 				const EInputEvent InputEvent = IE_Pressed;
-				UViewportInteractor::HandleInputKey( ControllerHandSide == EControllerHand::Left ? MotionController_Left_FullyPressedTriggerAxis : MotionController_Right_FullyPressedTriggerAxis, InputEvent );
+				UViewportInteractor::HandleInputKey( ViewportClient, ControllerHandSide == EControllerHand::Left ? MotionController_Left_FullyPressedTriggerAxis : MotionController_Right_FullyPressedTriggerAxis, InputEvent );
 			}
 			else if ( bIsTriggerFullyPressed && Delta < VREd::TriggerFullyPressedReleaseThreshold->GetFloat() )
 			{
@@ -723,7 +722,7 @@ void UVREditorMotionControllerInteractor::HandleInputAxis( FViewportActionKeyInp
 
 				// Synthesize an input key for this full press
 				const EInputEvent InputEvent = IE_Released;
-				UViewportInteractor::HandleInputKey( ControllerHandSide == EControllerHand::Left ? MotionController_Left_FullyPressedTriggerAxis : MotionController_Right_FullyPressedTriggerAxis, InputEvent );
+				UViewportInteractor::HandleInputKey( ViewportClient, ControllerHandSide == EControllerHand::Left ? MotionController_Left_FullyPressedTriggerAxis : MotionController_Right_FullyPressedTriggerAxis, InputEvent );
 			}
 		}
 	}
@@ -753,7 +752,7 @@ void UVREditorMotionControllerInteractor::HandleInputAxis( FViewportActionKeyInp
 		}
 	}
 
-	Super::HandleInputAxis( Action, Key, Delta, DeltaTime, bOutWasHandled );
+	Super::HandleInputAxis( ViewportClient, Action, Key, Delta, DeltaTime, bOutWasHandled );
 }
 
 void UVREditorMotionControllerInteractor::PollInput()
@@ -901,7 +900,7 @@ void UVREditorMotionControllerInteractor::StopOldHapticEffects()
 	}
 }
 
-bool UVREditorMotionControllerInteractor::GetTransformAndForwardVector( FTransform& OutHandTransform, FVector& OutForwardVector )
+bool UVREditorMotionControllerInteractor::GetTransformAndForwardVector( FTransform& OutHandTransform, FVector& OutForwardVector ) const
 {
 	if ( bHaveMotionController )
 	{
@@ -1108,7 +1107,7 @@ void UVREditorMotionControllerInteractor::ShowHelpForHand( const bool bShowIt )
 
 					const bool bWithSceneComponent = false;	// Nope, we'll spawn our own inside AFloatingText
 					check( VRMode );
-					AFloatingText* FloatingText = GetVRMode().SpawnTransientSceneActor<AFloatingText>( ComponentName, bWithSceneComponent );
+					AFloatingText* FloatingText = UViewportWorldInteraction::SpawnTransientSceneActor<AFloatingText>( GetWorld(), ComponentName, bWithSceneComponent );
 					FloatingText->SetText( LabelText );
 
 					HelpLabels.Add( Key, FloatingText );
@@ -1147,7 +1146,7 @@ void UVREditorMotionControllerInteractor::UpdateHelpLabels()
 		for ( auto& KeyAndValue : HelpLabels )
 		{
 			AFloatingText* FloatingText = KeyAndValue.Value;
-			GetVRMode().DestroyTransientActor( FloatingText );
+			UViewportWorldInteraction::DestroyTransientActor( GetWorld(), FloatingText );
 		}
 		HelpLabels.Reset();
 	}
@@ -1314,9 +1313,9 @@ void UVREditorMotionControllerInteractor::UpdateRadialMenuInput( const float Del
 	EViewportInteractionDraggingMode DraggingMode = GetDraggingMode();
 	if ( !HasUIInFront() &&
 		( bIsTrackpadPositionValid[ 0 ] && bIsTrackpadPositionValid[ 1 ] ) &&
-		DraggingMode != EViewportInteractionDraggingMode::ActorsWithGizmo &&
-		DraggingMode != EViewportInteractionDraggingMode::ActorsFreely &&
-		DraggingMode != EViewportInteractionDraggingMode::ActorsAtLaserImpact &&
+		DraggingMode != EViewportInteractionDraggingMode::TransformablesWithGizmo &&
+		DraggingMode != EViewportInteractionDraggingMode::TransformablesFreely &&
+		DraggingMode != EViewportInteractionDraggingMode::TransformablesAtLaserImpact &&
 		DraggingMode != EViewportInteractionDraggingMode::AssistingDrag &&
 		!UISystem.IsInteractorDraggingDockUI( this ) &&
 		!UISystem.IsShowingRadialMenu( Cast<UVREditorInteractor>( OtherInteractor ) ) )
diff --git a/Engine/Source/Editor/VREditor/VREditorMotionControllerInteractor.h b/Engine/Source/Editor/VREditor/VREditorMotionControllerInteractor.h
--- a/Engine/Source/Editor/VREditor/VREditorMotionControllerInteractor.h
+++ b/Engine/Source/Editor/VREditor/VREditorMotionControllerInteractor.h
@@ -45,8 +45,8 @@ public:
 	EHMDDeviceType::Type GetHMDDeviceType() const;
 
 	// UViewportInteractor
-	virtual void HandleInputKey( FViewportActionKeyInput& Action, const FKey Key, const EInputEvent Event, bool& bOutWasHandled ) override;
-	virtual bool GetTransformAndForwardVector( FTransform& OutHandTransform, FVector& OutForwardVector ) override;
+	virtual void HandleInputKey( class FEditorViewportClient& ViewportClient, FViewportActionKeyInput& Action, const FKey Key, const EInputEvent Event, bool& bOutWasHandled ) override;
+	virtual bool GetTransformAndForwardVector( FTransform& OutHandTransform, FVector& OutForwardVector ) const override;
 
 	/** Gets the Y delta of the trackpad */
 	float GetTrackpadSlideDelta();
@@ -81,7 +81,7 @@ public:
 protected:
 
 	// ViewportInteractor
-	virtual void HandleInputAxis( FViewportActionKeyInput& Action, const FKey Key, const float Delta, const float DeltaTime, bool& bOutWasHandled ) override;
+	virtual void HandleInputAxis( class FEditorViewportClient& ViewportClient, FViewportActionKeyInput& Action, const FKey Key, const float Delta, const float DeltaTime, bool& bOutWasHandled ) override;
 
 	/** Polls input for the motion controllers transforms */
 	virtual void PollInput() override;
diff --git a/Engine/Source/Editor/VREditor/VREditorWorldInteraction.cpp b/Engine/Source/Editor/VREditor/VREditorWorldInteraction.cpp
--- a/Engine/Source/Editor/VREditor/VREditorWorldInteraction.cpp
+++ b/Engine/Source/Editor/VREditor/VREditorWorldInteraction.cpp
@@ -119,8 +119,8 @@ void UVREditorWorldInteraction::StopDragging( UViewportInteractor* Interactor )
 
 	// If we were placing something, bring the window back
 	if ( FloatingUIAssetDraggedFrom != nullptr && 
-		( InteractorDraggingMode == EViewportInteractionDraggingMode::ActorsAtLaserImpact ||
-		InteractorDraggingMode == EViewportInteractionDraggingMode::Material ) )
+		( InteractorDraggingMode == EViewportInteractionDraggingMode::TransformablesAtLaserImpact ||
+		  InteractorDraggingMode == EViewportInteractionDraggingMode::Material ) )
 	{
 		Owner->GetUISystem().ShowEditorUIPanel( FloatingUIAssetDraggedFrom, Cast<UVREditorInteractor>( Interactor->GetOtherInteractor() ), 
 			true, false, false, false );
@@ -264,7 +264,7 @@ void UVREditorWorldInteraction::OnAssetDragStartedFromContentBrowser( const TArr
 			bool bCanPlace = true;
 
 			UObject* AssetObj = AssetData.GetAsset();
-			if( !ObjectTools::IsAssetValidForPlacing( ViewportWorldInteraction->GetViewportWorld(), AssetData.ObjectPath.ToString() ) )
+			if( !ObjectTools::IsAssetValidForPlacing( ViewportWorldInteraction->GetWorld(), AssetData.ObjectPath.ToString() ) )
 			{
 				bCanPlace = false;
 			}
@@ -320,7 +320,7 @@ void UVREditorWorldInteraction::OnAssetDragStartedFromContentBrowser( const TArr
 				const bool bSelectNewActors = true;
 				const EObjectFlags NewObjectFlags = bIsPreview ? RF_Transient : RF_Transactional;
 
-				TArray<AActor*> NewActors = FLevelEditorViewportClient::TryPlacingActorFromObject( ViewportWorldInteraction->GetViewportWorld()->GetCurrentLevel(), AssetObj, bSelectNewActors, NewObjectFlags, FactoryToUse );
+				TArray<AActor*> NewActors = FLevelEditorViewportClient::TryPlacingActorFromObject( ViewportWorldInteraction->GetWorld()->GetCurrentLevel(), AssetObj, bSelectNewActors, NewObjectFlags, FactoryToUse );
 
 				if( NewActors.Num() > 0 )
 				{
@@ -370,32 +370,30 @@ void UVREditorWorldInteraction::OnAssetDragStartedFromContentBrowser( const TArr
 
 		if( AllNewActors.Num() > 0 )
 		{
-			const FViewportActionKeyInput Action( ViewportWorldActionTypes::SelectAndMove_LightlyPressed );
-
 			// Start dragging the new actor(s)
-			const bool bIsPlacingActors = true;
-			ViewportWorldInteraction->StartDraggingActors( PlacingWithInteractor, Action, ViewportWorldInteraction->GetTransformGizmoActor()->GetRootComponent(), PlaceAt, bIsPlacingActors );
-
-			// If we're interpolating, update the target transform of the actors to use our overridden size.  When
-			// we placed them we set their size to be 'thumbnail sized', and we want them to interpolate to
-			// their actual size in the world
-			if( bShouldInterpolateFromDragLocation )
+			UPrimitiveComponent* ClickedTransformGizmoComponent = nullptr;
+			const bool bIsPlacingNewObjects = true;
+			const bool bStartTransaction = false;
+			if( ViewportWorldInteraction->StartDragging( PlacingWithInteractor, ClickedTransformGizmoComponent, PlaceAt, bIsPlacingNewObjects, bStartTransaction ) )
 			{
-				for( TUniquePtr<FViewportTransformable>& TransformablePtr : ViewportWorldInteraction->Transformables )
-				{
-					FViewportTransformable& Transformable = *TransformablePtr;
-
-					float ObjectSize = 1.0f;
+				// Dragging started!
 
-					if( VREd::OverrideSizeOfPlacedActors->GetFloat() > KINDA_SMALL_NUMBER )
+				// If we're interpolating, update the target transform of the actors to use our overridden size.  When
+				// we placed them we set their size to be 'thumbnail sized', and we want them to interpolate to
+				// their actual size in the world
+				if( bShouldInterpolateFromDragLocation )
+				{
+					for( TUniquePtr<FViewportTransformable>& TransformablePtr : ViewportWorldInteraction->GetTransformables() )
 					{
-						ObjectSize = 1.0f * ViewportWorldInteraction->GetWorldScaleFactor();
+						FViewportTransformable& Transformable = *TransformablePtr;
+
+						float ObjectSize = 1.0f;
 
-						// Got an actor?
-						AActor* Actor = Cast<AActor>( Transformable.Object.Get() );
-						if( Actor != nullptr )
+						if( VREd::OverrideSizeOfPlacedActors->GetFloat() > KINDA_SMALL_NUMBER )
 						{
-							const FBox LocalSpaceBounds = Actor->CalculateComponentsBoundingBoxInLocalSpace();
+							ObjectSize = 1.0f * ViewportWorldInteraction->GetWorldScaleFactor();
+
+							const FBox LocalSpaceBounds = Transformable.GetLocalSpaceBoundingBox();
 							const float LocalBoundsSize = LocalSpaceBounds.GetSize().GetAbsMax();
 
 							const float DesiredSize = VREd::OverrideSizeOfPlacedActors->GetFloat() * ViewportWorldInteraction->GetWorldScaleFactor();
@@ -470,7 +468,7 @@ void UVREditorWorldInteraction::PlaceDraggedMaterialOrTexture( UViewportInteract
 
 			if( DroppedObjAsMaterial || DroppedObjAsTexture )
 			{
-				UGameplayStatics::PlaySound2D( ViewportWorldInteraction->GetViewportWorld(), DropMaterialOrMaterialSound );
+				UGameplayStatics::PlaySound2D( ViewportWorldInteraction->GetWorld(), DropMaterialOrMaterialSound );
 			}
 		}
 	}
@@ -480,7 +478,7 @@ void UVREditorWorldInteraction::PlaceDraggedMaterialOrTexture( UViewportInteract
 
 void UVREditorWorldInteraction::SnapSelectedActorsToGround()
 {
-	TSharedPtr<SLevelViewport> LevelEditorViewport = StaticCastSharedPtr<SLevelViewport>( ViewportWorldInteraction->GetViewportClient()->GetEditorViewportWidget() );
+	TSharedPtr<SLevelViewport> LevelEditorViewport = StaticCastSharedPtr<SLevelViewport>( ViewportWorldInteraction->GetDefaultOptionalViewportClient()->GetEditorViewportWidget() );
 	if ( LevelEditorViewport.IsValid() )
 	{
 		FLevelEditorModule& LevelEditorModule = FModuleManager::GetModuleChecked<FLevelEditorModule>( TEXT( "LevelEditor" ) );
@@ -489,8 +487,8 @@ void UVREditorWorldInteraction::SnapSelectedActorsToGround()
 
 		CommandList->ExecuteAction( Commands.SnapBottomCenterBoundsToFloor.ToSharedRef() );
 
-		// @todo vreditor: This should not be needed after to allow transformables to stop animating the transforms of actors after they come to rest
-		ViewportWorldInteraction->SetupTransformablesForSelectedActors();
+		// Selected objects have been moved, so make sure our transformables are all reset to the object's new location
+		ViewportWorldInteraction->ResetTransformables();
 	}
 }
 
diff --git a/Engine/Source/Editor/ViewportInteraction/ActorTransformer.cpp b/Engine/Source/Editor/ViewportInteraction/ActorTransformer.cpp
new file mode 100644
--- /dev/null
+++ b/Engine/Source/Editor/ViewportInteraction/ActorTransformer.cpp
@@ -0,0 +1,50 @@
+// Copyright 1998-2016 Epic Games, Inc. All Rights Reserved.
+
+#include "ActorTransformer.h"
+#include "ActorViewportTransformable.h"
+#include "ViewportWorldInteraction.h"
+#include "Engine/Selection.h"
+
+
+void UActorTransformer::Init( UViewportWorldInteraction* InitViewportWorldInteraction )
+{
+	Super::Init( InitViewportWorldInteraction );
+
+	// Find out about selection changes
+	USelection::SelectionChangedEvent.AddUObject( this, &UActorTransformer::OnActorSelectionChanged );
+}
+
+
+void UActorTransformer::Shutdown()
+{
+	USelection::SelectionChangedEvent.RemoveAll( this );
+
+	Super::Shutdown();
+}
+
+
+void UActorTransformer::OnActorSelectionChanged( UObject* ChangedObject )
+{
+	TArray<TUniquePtr<FViewportTransformable>> NewTransformables;
+
+	USelection* ActorSelectionSet = GEditor->GetSelectedActors();
+
+	static TArray<UObject*> SelectedActorObjects;
+	SelectedActorObjects.Reset();
+	ActorSelectionSet->GetSelectedObjects( AActor::StaticClass(), /* Out */ SelectedActorObjects );
+
+	for( UObject* SelectedActorObject : SelectedActorObjects )
+	{
+		AActor* SelectedActor = CastChecked<AActor>( SelectedActorObject );
+
+		FActorViewportTransformable* Transformable = new FActorViewportTransformable();
+
+		NewTransformables.Add( TUniquePtr<FViewportTransformable>( Transformable ) );
+
+		Transformable->ActorWeakPtr = SelectedActor;
+		Transformable->StartTransform = Transformable->LastTransform = Transformable->TargetTransform = Transformable->UnsnappedTargetTransform = Transformable->InterpolationSnapshotTransform = SelectedActor->GetTransform();
+	}
+
+	ViewportWorldInteraction->SetTransformables( MoveTemp( NewTransformables ) );
+}
+
diff --git a/Engine/Source/Editor/ViewportInteraction/ActorViewportTransformable.cpp b/Engine/Source/Editor/ViewportInteraction/ActorViewportTransformable.cpp
new file mode 100644
--- /dev/null
+++ b/Engine/Source/Editor/ViewportInteraction/ActorViewportTransformable.cpp
@@ -0,0 +1,101 @@
+// Copyright 1998-2016 Epic Games, Inc. All Rights Reserved.
+
+#include "ActorViewportTransformable.h"
+
+
+void FActorViewportTransformable::ApplyTransform( const FTransform& NewTransform, const bool bSweep )
+{
+	AActor* Actor = ActorWeakPtr.Get();
+	if( Actor != nullptr )
+	{
+		const FTransform& ExistingTransform = Actor->GetActorTransform();
+		if( !ExistingTransform.Equals( NewTransform, 0.0f ) )
+		{
+			const bool bOnlyTranslationChanged =
+				ExistingTransform.GetRotation() == NewTransform.GetRotation() &&
+				ExistingTransform.GetScale3D() == NewTransform.GetScale3D();
+
+			Actor->SetActorTransform( NewTransform, bSweep );
+			//GWarn->Logf( TEXT( "SMOOTH: Actor %s to %s" ), *Actor->GetName(), *Transformable.TargetTransform.ToString() );
+
+			Actor->InvalidateLightingCacheDetailed( bOnlyTranslationChanged );
+
+			const bool bFinished = false;	// @todo gizmo now: PostEditChange never called; and bFinished=true never known!!
+			Actor->PostEditMove( bFinished );
+		}
+	}
+}
+
+
+const FTransform FActorViewportTransformable::GetTransform() const
+{
+	AActor* Actor = ActorWeakPtr.Get();
+	if( Actor != nullptr )
+	{
+		return Actor->GetTransform();
+	}
+	else
+	{
+		return FTransform::Identity;
+	}
+}
+
+
+FBox FActorViewportTransformable::GetLocalSpaceBoundingBox() const
+{
+	FBox LocalSpaceBounds( 0 );
+
+	AActor* Actor = Cast<AActor>( ActorWeakPtr.Get() );
+	if( Actor != nullptr )
+	{
+		const bool bIncludeNonCollidingComponents = false;	// @todo gizmo: Disabled this because it causes lights to have huge bounds
+		LocalSpaceBounds = Actor->CalculateComponentsBoundingBoxInLocalSpace( bIncludeNonCollidingComponents );
+	}
+	return LocalSpaceBounds;
+}
+
+
+bool FActorViewportTransformable::IsPhysicallySimulated() const
+{
+	bool bIsPhysicallySimulated = false;
+
+	AActor* Actor = Cast<AActor>( ActorWeakPtr.Get() );
+	if( Actor != nullptr )
+	{
+		UPrimitiveComponent* RootPrimitiveComponent = Cast<UPrimitiveComponent>( Actor->GetRootComponent() );
+		if( RootPrimitiveComponent != nullptr )
+		{
+			if( RootPrimitiveComponent->IsSimulatingPhysics() )
+			{
+				bIsPhysicallySimulated = true;
+			}
+		}
+	}
+
+	return bIsPhysicallySimulated;
+}
+
+
+void FActorViewportTransformable::SetLinearVelocity( const FVector& NewVelocity )
+{
+	AActor* Actor = Cast<AActor>( ActorWeakPtr.Get() );
+	if( Actor != nullptr )
+	{
+		UPrimitiveComponent* RootPrimitiveComponent = Cast<UPrimitiveComponent>( Actor->GetRootComponent() );
+		if( RootPrimitiveComponent != nullptr )
+		{
+			RootPrimitiveComponent->SetAllPhysicsLinearVelocity( NewVelocity );
+		}
+	}
+}
+
+
+void FActorViewportTransformable::UpdateIgnoredActorList( TArray<AActor*>& IgnoredActors )
+{
+	AActor* Actor = Cast<AActor>( ActorWeakPtr.Get() );
+	if( Actor != nullptr )
+	{
+		IgnoredActors.Add( Actor );
+	}
+}
+
diff --git a/Engine/Source/Editor/ViewportInteraction/Gizmo/VIBaseTransformGizmo.cpp b/Engine/Source/Editor/ViewportInteraction/Gizmo/VIBaseTransformGizmo.cpp
--- a/Engine/Source/Editor/ViewportInteraction/Gizmo/VIBaseTransformGizmo.cpp
+++ b/Engine/Source/Editor/ViewportInteraction/Gizmo/VIBaseTransformGizmo.cpp
@@ -9,9 +9,9 @@
 
 namespace VREd
 {
-	static FAutoConsoleVariable GizmoSelectionAnimationDuration( TEXT( "VREd.GizmoSelectionAnimationDuration" ), 0.15f, TEXT( "How long to animate the gizmo after objects are selected" ) );
-	static FAutoConsoleVariable GizmoSelectionAnimationCurvePower( TEXT( "VREd.GizmoSelectionAnimationCurvePower" ), 2.0f, TEXT( "Controls the animation curve for the gizmo after objects are selected" ) );
-	static FAutoConsoleVariable GizmoShowMeasurementText( TEXT( "VREd.GizmoShowMeasurementText" ), 0, TEXT( "When enabled, gizmo measurements will always be visible.  Otherwise, only when hovering over a scale/stretch gizmo handle" ) );
+	static FAutoConsoleVariable GizmoSelectionAnimationDuration( TEXT( "VI.GizmoSelectionAnimationDuration" ), 0.15f, TEXT( "How long to animate the gizmo after objects are selected" ) );
+	static FAutoConsoleVariable GizmoSelectionAnimationCurvePower( TEXT( "VI.GizmoSelectionAnimationCurvePower" ), 2.0f, TEXT( "Controls the animation curve for the gizmo after objects are selected" ) );
+	static FAutoConsoleVariable GizmoShowMeasurementText( TEXT( "VI.GizmoShowMeasurementText" ), 0, TEXT( "When enabled, gizmo measurements will always be visible.  Otherwise, only when hovering over a scale/stretch gizmo handle" ) );
 }
 
 ABaseTransformGizmo::ABaseTransformGizmo( ) :
diff --git a/Engine/Source/Editor/ViewportInteraction/Gizmo/VIPivotTransformGizmo.cpp b/Engine/Source/Editor/ViewportInteraction/Gizmo/VIPivotTransformGizmo.cpp
--- a/Engine/Source/Editor/ViewportInteraction/Gizmo/VIPivotTransformGizmo.cpp
+++ b/Engine/Source/Editor/ViewportInteraction/Gizmo/VIPivotTransformGizmo.cpp
@@ -11,12 +11,13 @@
 namespace VREd //@todo VREditor: Duplicates of TransformGizmo
 {
 	// @todo vreditor tweak: Tweak out console variables
-	static FAutoConsoleVariable PivotGizmoDistanceScaleFactor( TEXT( "VREd.PivotGizmoDistanceScaleFactor" ), 0.003f, TEXT( "How much the gizmo handles should increase in size with distance from the camera, to make it easier to select"));
-	static FAutoConsoleVariable PivotGizmoTranslationPivotOffsetX( TEXT("VREd.PivotGizmoTranslationPivotOffsetX" ), 20.0f, TEXT( "How much the translation cylinder is offsetted from the pivot" ) );
-	static FAutoConsoleVariable PivotGizmoScalePivotOffsetX( TEXT( "VREd.PivotGizmoScalePivotOffsetX" ), 120.0f, TEXT( "How much the non-uniform scale is offsetted from the pivot" ) );
-	static FAutoConsoleVariable PivotGizmoPlaneTranslationPivotOffsetYZ(TEXT("VREd.PivotGizmoPlaneTranslationPivotOffsetYZ" ), 40.0f, TEXT( "How much the plane translation is offsetted from the pivot" ) );
-	static FAutoConsoleVariable PivotGizmoTranslationScaleMultiply( TEXT( "VREd.PivotGizmoTranslationScaleMultiply" ), 3.0f, TEXT( "Multiplies translation handles scale" ) );
-	static FAutoConsoleVariable PivotGizmoTranslationHoverScaleMultiply( TEXT( "VREd.PivotGizmoTranslationHoverScaleMultiply" ), 1.2f, TEXT( "Multiplies translation handles hover scale" ) );
+	static FAutoConsoleVariable PivotGizmoMinDistanceForScaling( TEXT( "VI.PivotGizmoMinDistanceForScaling" ), 100.0f, TEXT( "How far away the camera needs to be from an object before we'll start scaling it based on distance"));
+	static FAutoConsoleVariable PivotGizmoDistanceScaleFactor( TEXT( "VI.PivotGizmoDistanceScaleFactor" ), 0.003f, TEXT( "How much the gizmo handles should increase in size with distance from the camera, to make it easier to select" ) );
+	static FAutoConsoleVariable PivotGizmoTranslationPivotOffsetX( TEXT("VI.PivotGizmoTranslationPivotOffsetX" ), 20.0f, TEXT( "How much the translation cylinder is offsetted from the pivot" ) );
+	static FAutoConsoleVariable PivotGizmoScalePivotOffsetX( TEXT( "VI.PivotGizmoScalePivotOffsetX" ), 120.0f, TEXT( "How much the non-uniform scale is offsetted from the pivot" ) );
+	static FAutoConsoleVariable PivotGizmoPlaneTranslationPivotOffsetYZ(TEXT("VI.PivotGizmoPlaneTranslationPivotOffsetYZ" ), 40.0f, TEXT( "How much the plane translation is offsetted from the pivot" ) );
+	static FAutoConsoleVariable PivotGizmoTranslationScaleMultiply( TEXT( "VI.PivotGizmoTranslationScaleMultiply" ), 3.0f, TEXT( "Multiplies translation handles scale" ) );
+	static FAutoConsoleVariable PivotGizmoTranslationHoverScaleMultiply( TEXT( "VI.PivotGizmoTranslationHoverScaleMultiply" ), 1.2f, TEXT( "Multiplies translation handles hover scale" ) );
 }
 
 APivotTransformGizmo::APivotTransformGizmo() :
@@ -64,7 +65,7 @@ APivotTransformGizmo::APivotTransformGizmo() :
 	OnNewObjectsSelected();
 }
 
-void APivotTransformGizmo::UpdateGizmo( const EGizmoHandleTypes GizmoType, const ECoordSystem GizmoCoordinateSpace, const FTransform& LocalToWorld, const FBox& LocalBounds, const FVector ViewLocation, bool bAllHandlesVisible, 
+void APivotTransformGizmo::UpdateGizmo( const EGizmoHandleTypes GizmoType, const ECoordSystem GizmoCoordinateSpace, const FTransform& LocalToWorld, const FBox& LocalBounds, const FVector ViewLocation, const float ScaleMultiplier, bool bAllHandlesVisible, 
 	UActorComponent* DraggingHandle, const TArray< UActorComponent* >& HoveringOverHandles, const float GizmoHoverScale, const float GizmoHoverAnimationDuration )
 {
 	const float WorldScaleFactor = GetWorld()->GetWorldSettings()->WorldToMeters / 100.0f;
@@ -74,10 +75,10 @@ void APivotTransformGizmo::UpdateGizmo( const EGizmoHandleTypes GizmoType, const
 	this->SetActorTransform( LocalToWorld, bSweep );
 
 	// Increase scale with distance, to make gizmo handles easier to click on
-	// @todo vreditor: Should probably be a curve, not linear
-	// @todo vreditor: Should take FOV into account (especially in non-stereo/HMD mode)
-	const float WorldSpaceDistanceToToPivot = FMath::Sqrt( FVector::DistSquared( GetActorLocation(), ViewLocation ) );
-	const float GizmoScale( ( ( WorldSpaceDistanceToToPivot / WorldScaleFactor ) * VREd::PivotGizmoDistanceScaleFactor->GetFloat() ) * WorldScaleFactor );
+	const float WorldSpaceDistanceToToPivot = FMath::Max( VREd::PivotGizmoMinDistanceForScaling->GetFloat(), FMath::Sqrt( FVector::DistSquared( GetActorLocation(), ViewLocation ) ) );
+
+	// @todo gizmo: Scale causes bounds-based features to look wrong
+	const float GizmoScale( ScaleMultiplier * ( ( WorldSpaceDistanceToToPivot / WorldScaleFactor ) * VREd::PivotGizmoDistanceScaleFactor->GetFloat() ) * WorldScaleFactor );
 
 	// Update animation
 	float AnimationAlpha = GetAnimationAlpha();
diff --git a/Engine/Source/Editor/ViewportInteraction/Gizmo/VIPivotTransformGizmo.h b/Engine/Source/Editor/ViewportInteraction/Gizmo/VIPivotTransformGizmo.h
--- a/Engine/Source/Editor/ViewportInteraction/Gizmo/VIPivotTransformGizmo.h
+++ b/Engine/Source/Editor/ViewportInteraction/Gizmo/VIPivotTransformGizmo.h
@@ -24,7 +24,7 @@ public:
 	/** Called by the world interaction system after we've been spawned into the world, to allow
 	    us to create components and set everything up nicely for the selected objects that we'll be
 		used to manipulate */
-	virtual void UpdateGizmo( const EGizmoHandleTypes GizmoType, const ECoordSystem GizmoCoordinateSpace, const FTransform& LocalToWorld, const FBox& LocalBounds, const FVector ViewLocation, 
+	virtual void UpdateGizmo( const EGizmoHandleTypes GizmoType, const ECoordSystem GizmoCoordinateSpace, const FTransform& LocalToWorld, const FBox& LocalBounds, const FVector ViewLocation, const float ScaleMultiplier,
 		bool bAllHandlesVisible, class UActorComponent* DraggingHandle, const TArray< UActorComponent* >& HoveringOverHandles, const float GizmoHoverScale, const float GizmoHoverAnimationDuration ) override;
 
 private:
diff --git a/Engine/Source/Editor/ViewportInteraction/MouseCursorInteractor.cpp b/Engine/Source/Editor/ViewportInteraction/MouseCursorInteractor.cpp
--- a/Engine/Source/Editor/ViewportInteraction/MouseCursorInteractor.cpp
+++ b/Engine/Source/Editor/ViewportInteraction/MouseCursorInteractor.cpp
@@ -9,6 +9,9 @@
 UMouseCursorInteractor::UMouseCursorInteractor( const FObjectInitializer& Initializer ) 
 	: UViewportInteractor( Initializer )
 {
+	// Grabber spheres don't really work well with mouse cursor interactors, because the origin of the
+	// interactor is right on the near view plane.
+	bAllowGrabberSphere = false;
 }
 
 
@@ -26,32 +29,47 @@ void UMouseCursorInteractor::Init()
 
 void UMouseCursorInteractor::PollInput()
 {
-	FEditorViewportClient& ViewportClient = *WorldInteraction->GetViewportClient();
-
 	InteractorData.LastTransform = InteractorData.Transform;
 	InteractorData.LastRoomSpaceTransform = InteractorData.RoomSpaceTransform;
 
-	// Only if we're not tracking (RMB looking)
-	if( !ViewportClient.IsTracking() )
+	// Make sure we have a valid viewport with a cursor over it, and that the viewport's world is the same as ours
+	FEditorViewportClient* ViewportClientPtr = GCurrentLevelEditingViewportClient;
+	if( ViewportClientPtr != nullptr && ViewportClientPtr->GetWorld() == WorldInteraction->GetWorld() )
 	{
-		FViewport* Viewport = ViewportClient.Viewport;
+		FEditorViewportClient& ViewportClient = *ViewportClientPtr;
 
-		// Make sure we have a valid viewport, otherwise we won't be able to construct an FSceneView.  The first time we're ticked we might not be properly setup. (@todo mesheditor)
-		if( Viewport != nullptr && Viewport->GetSizeXY().GetMin() > 0 )
+		// Only if we're not tracking (RMB looking)
+		if( !ViewportClient.IsTracking() )
 		{
-			const int32 ViewportInteractX = Viewport->GetMouseX();
-			const int32 ViewportInteractY = Viewport->GetMouseY();
+			FViewport* Viewport = ViewportClient.Viewport;
+
+			// Make sure we have a valid viewport, otherwise we won't be able to construct an FSceneView.  The first time we're ticked we might not be properly setup. (@todo viewportinteraction)
+			if( Viewport != nullptr && Viewport->GetSizeXY().GetMin() > 0 )
+			{
+				const int32 ViewportInteractX = Viewport->GetMouseX();
+				const int32 ViewportInteractY = Viewport->GetMouseY();
+
+				FSceneViewFamilyContext ViewFamily( FSceneViewFamily::ConstructionValues(
+					Viewport,
+					ViewportClient.GetScene(),
+					ViewportClient.EngineShowFlags )
+					.SetRealtimeUpdate( ViewportClient.IsRealtime() ) );
+				FSceneView* SceneView = ViewportClient.CalcSceneView( &ViewFamily );
+
+				const FViewportCursorLocation MouseViewportRay( SceneView, &ViewportClient, ViewportInteractX, ViewportInteractY );
 
-			FSceneViewFamilyContext ViewFamily( FSceneViewFamily::ConstructionValues(
-				Viewport,
-				ViewportClient.GetScene(),
-				ViewportClient.EngineShowFlags )
-				.SetRealtimeUpdate( ViewportClient.IsRealtime() ) );
-			FSceneView* SceneView = ViewportClient.CalcSceneView( &ViewFamily );
+				FVector RayOrigin = MouseViewportRay.GetOrigin();
 
-			const FViewportCursorLocation MouseViewportRay( SceneView, &ViewportClient, ViewportInteractX, ViewportInteractY );
+				// If we're dealing with an orthographic view, push the origin of the ray backward along the viewport forward axis
+				// to make sure that we can select objects that are behind the origin!
+				if( !ViewportClient.IsPerspective() )
+				{
+					const float HalfLaserPointerLength = this->GetLaserPointerMaxLength() * 0.5f;
+					RayOrigin -= MouseViewportRay.GetDirection() * HalfLaserPointerLength;
+				}
 
-			InteractorData.Transform = FTransform( MouseViewportRay.GetDirection().ToOrientationQuat(), MouseViewportRay.GetOrigin(), FVector( 1.0f ) );
+				InteractorData.Transform = FTransform( MouseViewportRay.GetDirection().ToOrientationQuat(), RayOrigin, FVector( 1.0f ) );
+			}
 		}
 	}
 	InteractorData.RoomSpaceTransform = InteractorData.Transform * WorldInteraction->GetRoomTransform().Inverse();
diff --git a/Engine/Source/Editor/ViewportInteraction/Public/ActorTransformer.h b/Engine/Source/Editor/ViewportInteraction/Public/ActorTransformer.h
new file mode 100644
--- /dev/null
+++ b/Engine/Source/Editor/ViewportInteraction/Public/ActorTransformer.h
@@ -0,0 +1,29 @@
+// Copyright 1998-2016 Epic Games, Inc. All Rights Reserved.
+
+#pragma once
+
+#include "ViewportTransformer.h"
+#include "ActorTransformer.generated.h"
+
+
+UCLASS()
+class VIEWPORTINTERACTION_API UActorTransformer : public UViewportTransformer
+{
+	GENERATED_BODY()
+
+public:
+
+	UFUNCTION()
+	virtual void Init( class UViewportWorldInteraction* InitViewportWorldInteraction ) override;
+
+	UFUNCTION()
+	virtual void Shutdown() override;
+
+
+protected:
+
+	/** Called when selection changes in the level */
+	void OnActorSelectionChanged( UObject* ChangedObject );
+
+};
+
diff --git a/Engine/Source/Editor/ViewportInteraction/Public/ActorViewportTransformable.h b/Engine/Source/Editor/ViewportInteraction/Public/ActorViewportTransformable.h
new file mode 100644
--- /dev/null
+++ b/Engine/Source/Editor/ViewportInteraction/Public/ActorViewportTransformable.h
@@ -0,0 +1,33 @@
+// Copyright 1998-2016 Epic Games, Inc. All Rights Reserved.
+
+#pragma once
+
+#include "ViewportTransformable.h"
+
+/**
+ * A transformable actor
+ */
+class VIEWPORTINTERACTION_API FActorViewportTransformable : public FViewportTransformable
+{
+
+public:
+
+	/** Sets up safe defaults */
+	FActorViewportTransformable()
+		: ActorWeakPtr()
+	{
+	}
+
+	// FViewportTransformable overrides
+	virtual const FTransform GetTransform() const override;
+	virtual void ApplyTransform( const FTransform& NewTransform, const bool bSweep ) override;
+	virtual FBox GetLocalSpaceBoundingBox() const override;
+	virtual bool IsPhysicallySimulated() const override;
+	virtual void SetLinearVelocity( const FVector& NewVelocity ) override;
+	virtual void UpdateIgnoredActorList( TArray<class AActor*>& IgnoredActors ) override;
+
+	/** The actual actor object */
+	TWeakObjectPtr<class AActor> ActorWeakPtr;
+};
+
+
diff --git a/Engine/Source/Editor/ViewportInteraction/Public/IViewportInteractionModule.h b/Engine/Source/Editor/ViewportInteraction/Public/IViewportInteractionModule.h
--- a/Engine/Source/Editor/ViewportInteraction/Public/IViewportInteractionModule.h
+++ b/Engine/Source/Editor/ViewportInteraction/Public/IViewportInteractionModule.h
@@ -39,10 +39,5 @@ public:
 	 * @return The ViewportWorldInteractionManager that owns entering and exiting the ViewportWorldInteraction
 	 */
 	virtual class FViewportWorldInteractionManager& GetWorldInteractionManager() = 0;
-
-	/**
-	 * Updates the ViewportWorldInteraction
-	 */
-	virtual void Tick( float DeltaTime ) = 0;
 };
 
diff --git a/Engine/Source/Editor/ViewportInteraction/Public/IViewportWorldInteractionManager.h b/Engine/Source/Editor/ViewportInteraction/Public/IViewportWorldInteractionManager.h
--- a/Engine/Source/Editor/ViewportInteraction/Public/IViewportWorldInteractionManager.h
+++ b/Engine/Source/Editor/ViewportInteraction/Public/IViewportWorldInteractionManager.h
@@ -7,14 +7,7 @@
 class IViewportWorldInteractionManager
 {
 public:
-	DECLARE_EVENT_OneParam( IViewportWorldInteractionManager, FOnWorldInteractionTick, const float /* DeltaTime */ );
-
-	/** Gets the event for ticking before the viewport world interaction */
-	virtual FOnWorldInteractionTick& OnPreWorldInteractionTick() = 0;
-
-	/** Gets the event for ticking after the viewport world interaction */
-	virtual FOnWorldInteractionTick& OnPostWorldInteractionTick() = 0;
-
+	
 	/** 
 	 * Sets the current ViewportWorldInteraction
 	 *
@@ -22,8 +15,4 @@ public:
 	 */
 	virtual void SetCurrentViewportWorldInteraction( class UViewportWorldInteraction* /*WorldInteraction*/ ) = 0;
 
-	/**
-	 * Updates the current ViewportWorldInteraction
-	 */
-	virtual void Tick( float DeltaTime ) = 0;
 };
diff --git a/Engine/Source/Editor/ViewportInteraction/Public/VIBaseTransformGizmo.h b/Engine/Source/Editor/ViewportInteraction/Public/VIBaseTransformGizmo.h
--- a/Engine/Source/Editor/ViewportInteraction/Public/VIBaseTransformGizmo.h
+++ b/Engine/Source/Editor/ViewportInteraction/Public/VIBaseTransformGizmo.h
@@ -77,7 +77,7 @@ public:
 	/** Called by the world interaction system after we've been spawned into the world, to allow
 	    us to create components and set everything up nicely for the selected objects that we'll be
 		used to manipulate */
-	virtual void UpdateGizmo( const EGizmoHandleTypes GizmoType, const ECoordSystem GizmoCoordinateSpace, const FTransform& LocalToWorld, const FBox& LocalBounds, const FVector ViewLocation, bool bAllHandlesVisible, 
+	virtual void UpdateGizmo( const EGizmoHandleTypes GizmoType, const ECoordSystem GizmoCoordinateSpace, const FTransform& LocalToWorld, const FBox& LocalBounds, const FVector ViewLocation, const float ScaleMultiplier, bool bAllHandlesVisible,
 		class UActorComponent* DraggingHandle, const TArray< UActorComponent* >& HoveringOverHandles, const float GizmoHoverScale, const float GizmoHoverAnimationDuration ) {}
 
 protected:
diff --git a/Engine/Source/Editor/ViewportInteraction/Public/ViewportInteraction.h b/Engine/Source/Editor/ViewportInteraction/Public/ViewportInteraction.h
--- a/Engine/Source/Editor/ViewportInteraction/Public/ViewportInteraction.h
+++ b/Engine/Source/Editor/ViewportInteraction/Public/ViewportInteraction.h
@@ -8,6 +8,5 @@ MONOLITHIC_HEADER_BOILERPLATE()
 #include "ViewportDragOperation.h"
 #include "ViewportInteractableInterface.h"
 #include "IViewportInteractionModule.h"
-#include "ViewportWorldInteractionInterface.h"
 #include "ViewportInteractionTypes.h"
 #include "ViewportInteractor.h"
diff --git a/Engine/Source/Editor/ViewportInteraction/Public/ViewportInteractionTypes.h b/Engine/Source/Editor/ViewportInteraction/Public/ViewportInteractionTypes.h
--- a/Engine/Source/Editor/ViewportInteraction/Public/ViewportInteractionTypes.h
+++ b/Engine/Source/Editor/ViewportInteraction/Public/ViewportInteractionTypes.h
@@ -44,17 +44,17 @@ enum class EViewportInteractionDraggingMode : uint8
 	/** Not dragging right now with this hand */
 	Nothing,
 
-	/** Dragging actors around using the transform gizmo */
-	ActorsWithGizmo,
+	/** Dragging transformables (e.g. actors, components, geometry elements) around using the transform gizmo */
+	TransformablesWithGizmo,
 
-	/** Actors locked to the impact point under the laser */
-	ActorsAtLaserImpact,
+	/** Transformables locked to the impact point under the laser */
+	TransformablesAtLaserImpact,
 
 	/** We're grabbing an object (or the world) that was already grabbed by the other hand */
 	AssistingDrag,
 
-	/** Freely moving, rotating and scaling objects with one or two hands */
-	ActorsFreely,
+	/** Freely moving, rotating and scaling transformables with one or two hands */
+	TransformablesFreely,
 
 	/** Moving the world itself around (actually, moving the camera in such a way that it feels like you're moving the world) */
 	World,
diff --git a/Engine/Source/Editor/ViewportInteraction/Public/ViewportInteractor.h b/Engine/Source/Editor/ViewportInteraction/Public/ViewportInteractor.h
--- a/Engine/Source/Editor/ViewportInteraction/Public/ViewportInteractor.h
+++ b/Engine/Source/Editor/ViewportInteraction/Public/ViewportInteractor.h
@@ -26,6 +26,9 @@ public:
 	/** Gets the private data for this interactor */
 	struct FViewportInteractorData& GetInteractorData();
 
+	/** Gets the private data for this interactor (const) */
+	const struct FViewportInteractorData& GetInteractorData() const;
+
 	/** Sets the world interaction */
 	void SetWorldInteraction( class UViewportWorldInteraction* InWorldInteraction );
 
@@ -60,10 +63,10 @@ public:
 	virtual void PollInput() {}
 
 	/** Handles axis input and translates it actions */
-	bool HandleInputKey( const FKey Key, const EInputEvent Event );
+	bool HandleInputKey( class FEditorViewportClient& ViewportClient, const FKey Key, const EInputEvent Event );
 
 	/** Handles axis input and translates it to actions */
-	bool HandleInputAxis( const FKey Key, const float Delta, const float DeltaTime );
+	bool HandleInputAxis( class FEditorViewportClient& ViewportClient, const FKey Key, const float Delta, const float DeltaTime );
 
 	/** Gets the world transform of this interactor */
 	FTransform GetTransform() const;
@@ -83,12 +86,22 @@ public:
 	 * @param HandIndex				Index of the hand to use
 	 * @param LasertPointerStart	(Out) The start location of the laser pointer in world space
 	 * @param LasertPointerEnd		(Out) The end location of the laser pointer in world space
-	 * @param bEvenIfUIIsInFront	If true, returns a laser pointer even if the hand has UI in front of it (defaults to false)
+	 * @param bEvenIfBlocked		If true, returns a laser pointer even if the hand has UI in front of it (defaults to false)
 	 * @param LaserLengthOverride	If zero the default laser length (VREdMode::GetLaserLength) is used
 	 *
 	 * @return	True if we have motion controller data for this hand and could return a valid result
 	 */
-	bool GetLaserPointer( FVector& LaserPointerStart, FVector& LaserPointerEnd, const bool bEvenIfBlocked = false, const float LaserLengthOverride = 0.0f );
+	bool GetLaserPointer( FVector& LaserPointerStart, FVector& LaserPointerEnd, const bool bEvenIfBlocked = false, const float LaserLengthOverride = 0.0f ) const;
+
+	/**
+	 * Gets a sphere on this interactor that can be used to interact with objects in close proximity
+	 *
+	 * @param	OutGrabberSphere	The sphere in world space
+	 * @param	bEvenIfBlocked	When set to true, a valid sphere may be returned even if there is UI attached in front of this interactor
+	 *
+	 * @return	True if the sphere is available and valid, or false is the interactor was busy or we could not determine a valid position for the interactor
+	 */
+	bool GetGrabberSphere( FSphere& OutGrabberSphere, const bool bEvenIfBlocked = false ) const;
 
 	/** Gets the maximum length of a laser pointer */
 	float GetLaserPointerMaxLength() const;
@@ -99,7 +112,6 @@ public:
 	/**
 	 * Traces along the laser pointer vector and returns what it first hits in the world
 	 *
-	 * @param HandIndex	Index of the hand that uses has the laser pointer to trace
 	 * @param OptionalListOfIgnoredActors Actors to exclude from hit testing
 	 * @param bIgnoreGizmos True if no gizmo results should be returned, otherwise they are preferred (x-ray)
 	 * @param bEvenIfUIIsInFront If true, ignores any UI that might be blocking the ray
@@ -125,10 +137,10 @@ public:
 	 *
 	 * @return	True if we have motion controller data for this hand and could return a valid result
 	 */
-	virtual bool GetTransformAndForwardVector( FTransform& OutHandTransform, FVector& OutForwardVector );
+	virtual bool GetTransformAndForwardVector( FTransform& OutHandTransform, FVector& OutForwardVector ) const;
 
-	/** Called by StartDragging in world interaction to give the interator a chance of acting upon starting a drag operation */
-	virtual void OnStartDragging( UActorComponent* ClickedComponent, const FVector& HitLocation, const bool bIsPlacingActors ) {};
+	/** Called by StartDragging in world interaction to give the interactor a chance of acting upon starting a drag operation */
+	virtual void OnStartDragging( const FVector& HitLocation, const bool bIsPlacingNewObjects ) {};
 
 	/** Gets the interactor laser hover location */
 	FVector GetHoverLocation() const;
@@ -143,7 +155,7 @@ public:
 	void SetDraggingMode( const EViewportInteractionDraggingMode NewDraggingMode );
 
 	/** To be overridden by base class. Called by GetLaserPointer to give the derived interactor a chance to disable the laser. By default it is not blocked */
-	virtual bool GetIsLaserBlocked() { return false; }
+	virtual bool GetIsLaserBlocked() const { return false; }
 
 	/** Gets a certain action by iterating through the map looking for the same ActionType */
 	// @todo ViewportInteractor: This should be changed to return a const pointer, but we need to fix up some dragging code in WorldInteraction first
@@ -183,10 +195,10 @@ public:
 protected:
 
 	/** To be overridden by base class. Called by HandleInputKey before delegates and default input implementation */
-	virtual void HandleInputKey( FViewportActionKeyInput& Action, const FKey Key, const EInputEvent Event, bool& bOutWasHandled ) {};
+	virtual void HandleInputKey( class FEditorViewportClient& ViewportClient, FViewportActionKeyInput& Action, const FKey Key, const EInputEvent Event, bool& bOutWasHandled ) {};
 	
 	/** To be overridden by base class. Called by HandleInputAxis before delegates and default input implementation */
-	virtual void HandleInputAxis( FViewportActionKeyInput& Action, const FKey Key, const float Delta, const float DeltaTime, bool& bOutWasHandled ) {};
+	virtual void HandleInputAxis( class FEditorViewportClient& ViewportClient, FViewportActionKeyInput& Action, const FKey Key, const float Delta, const float DeltaTime, bool& bOutWasHandled ) {};
 
 protected:
 
@@ -203,4 +215,7 @@ protected:
 	/** The paired interactor by the world interaction */
 	UPROPERTY()
 	UViewportInteractor* OtherInteractor;
-};
+
+	/** True if this interactor supports 'grabber sphere' interaction.  Usually disabled for mouse cursors */
+	bool bAllowGrabberSphere;
+};
\ No newline at end of file
diff --git a/Engine/Source/Editor/ViewportInteraction/Public/ViewportTransformable.h b/Engine/Source/Editor/ViewportInteraction/Public/ViewportTransformable.h
new file mode 100644
--- /dev/null
+++ b/Engine/Source/Editor/ViewportInteraction/Public/ViewportTransformable.h
@@ -0,0 +1,66 @@
+// Copyright 1998-2016 Epic Games, Inc. All Rights Reserved.
+
+#pragma once
+
+
+/**
+ * Represents an object that we're actively interacting with, such as a selected actor
+ */
+class VIEWPORTINTERACTION_API FViewportTransformable
+{
+
+public:
+
+	/** Sets up safe defaults */
+	FViewportTransformable()
+		: StartTransform( FTransform::Identity ),
+		  LastTransform( FTransform::Identity ),
+		  TargetTransform( FTransform::Identity ),
+		  UnsnappedTargetTransform( FTransform::Identity ),
+		  InterpolationSnapshotTransform( FTransform::Identity )
+	{
+	}
+
+	/** Gets the current transform of this object */
+	virtual const FTransform GetTransform() const = 0;
+
+	/** Updates the transform of the actual object */
+	virtual void ApplyTransform( const FTransform& NewTransform, const bool bSweep ) = 0;
+
+	/** Returns the bounding box of this transformable in local space */
+	virtual FBox GetLocalSpaceBoundingBox() const = 0;
+
+	/** Returns true if this transformable is a physically simulated kinematic object */
+	virtual bool IsPhysicallySimulated() const
+	{
+		return false;
+	}
+
+	/** For physically simulated objects, sets the new velocity of the object */
+	virtual void SetLinearVelocity( const FVector& NewVelocity )
+	{
+	}
+
+	/** For actor transformables, this will add it's actor to the incoming list */
+	virtual void UpdateIgnoredActorList( TArray<class AActor*>& IgnoredActors )
+	{
+	}
+
+
+	/** The object's position when we started the action */
+	FTransform StartTransform;
+
+	/** The last position we were at.  Used for smooth interpolation when snapping is enabled */
+	FTransform LastTransform;
+
+	/** The target position we want to be at.  Used for smooth interpolation when snapping is enabled */
+	FTransform TargetTransform;
+
+	/** The target position we'd be at if snapping wasn't enabled.  This is used for 'elastic' smoothing when
+	snapping is turned on.  Basically it allows us to 'reach' toward the desired position */
+	FTransform UnsnappedTargetTransform;
+
+	/** A "snapshot" transform, used for certain interpolation effects */
+	FTransform InterpolationSnapshotTransform;
+};
+
diff --git a/Engine/Source/Editor/ViewportInteraction/Public/ViewportTransformer.h b/Engine/Source/Editor/ViewportInteraction/Public/ViewportTransformer.h
new file mode 100644
--- /dev/null
+++ b/Engine/Source/Editor/ViewportInteraction/Public/ViewportTransformer.h
@@ -0,0 +1,29 @@
+// Copyright 1998-2016 Epic Games, Inc. All Rights Reserved.
+
+#pragma once
+
+#include "ViewportTransformer.generated.h"
+
+
+UCLASS( abstract )
+class VIEWPORTINTERACTION_API UViewportTransformer : public UObject
+{
+	GENERATED_BODY()
+
+public:
+
+	UFUNCTION()
+	virtual void Init( class UViewportWorldInteraction* InitViewportWorldInteraction );
+
+	UFUNCTION()
+	virtual void Shutdown();
+
+
+protected:
+
+	/** The viewport world interaction object we're registered with */
+	UPROPERTY()
+	class UViewportWorldInteraction* ViewportWorldInteraction;
+
+};
+
diff --git a/Engine/Source/Editor/ViewportInteraction/Public/ViewportWorldInteraction.h b/Engine/Source/Editor/ViewportInteraction/Public/ViewportWorldInteraction.h
--- a/Engine/Source/Editor/ViewportInteraction/Public/ViewportWorldInteraction.h
+++ b/Engine/Source/Editor/ViewportInteraction/Public/ViewportWorldInteraction.h
@@ -10,11 +10,12 @@
 #include "UObject/WeakObjectPtr.h"
 #include "GameFramework/Actor.h"
 #include "EditorViewportClient.h"
-#include "ViewportWorldInteractionInterface.h"
 #include "Misc/App.h"
 #include "ViewportInteractionTypes.h"
+#include "EditorWorldExtension.h"
 #include "VIBaseTransformGizmo.h"
 #include "LevelEditorViewport.h"
+#include "ViewportTransformable.h"
 #include "ViewportWorldInteraction.generated.h"
 
 class IViewportInteractableInterface;
@@ -36,86 +37,101 @@ enum class EGizmoHandleTypes : uint8;
 class IViewportInteractableInterface;
 class UViewportInteractor;
 
-/**
-* Represents an object that we're actively interacting with, such as a selected actor
-*/
-class FViewportTransformable
-{
 
+UCLASS()
+class VIEWPORTINTERACTION_API UViewportWorldInteraction : public UEditorWorldExtension
+{
+	GENERATED_BODY()
+	
 public:
 
-	/** Sets up safe defaults */
-	FViewportTransformable()
-		: Object(),
-		StartTransform( FTransform::Identity ),
-		LastTransform( FTransform::Identity ),
-		TargetTransform( FTransform::Identity ),
-		UnsnappedTargetTransform( FTransform::Identity ),
-		InterpolationSnapshotTransform( FTransform::Identity )
-	{
-	}
+	UViewportWorldInteraction();
 
-	/** The actual object */
-	FWeakObjectPtr Object;
+	// UEditorExtension overrides
+	virtual void Init() override;
+	virtual void Shutdown() override;
+	virtual void Tick( float DeltaSeconds ) override;
+	
+	/** Returns true if this worldinteraction is currently running */
+	bool IsActive() const;
 
-	/** The object's position when we started the action */
-	FTransform StartTransform;
+	/** Adds interactor to the worldinteraction */
+	void AddInteractor( UViewportInteractor* Interactor );
 
-	/** The last position we were at.  Used for smooth interpolation when snapping is enabled */
-	FTransform LastTransform;
+	/** Removes interactor from the worldinteraction and removes the interactor from its paired interactor if any */
+	void RemoveInteractor( UViewportInteractor* Interactor );
 
-	/** The target position we want to be at.  Used for smooth interpolation when snapping is enabled */
-	FTransform TargetTransform;
+	/** Creates an interactor for the mouse cursor.  If one already exists, this will add a reference to it.  Remember to
+	    call ReleaseMouseCursorInteractor() when you're done. */
+	void AddMouseCursorInteractor();
 
-	/** The target position we'd be at if snapping wasn't enabled.  This is used for 'elastic' smoothing when
-	snapping is turned on.  Basically it allows us to 'reach' toward the desired position */
-	FTransform UnsnappedTargetTransform;
+	/** When you're finished with the mouse cursor interactor, call this to release it.  It will be destroyed unless
+	    another system is still using it */
+	void ReleaseMouseCursorInteractor();
 
-	/** A "snapshot" transform, used for certain interpolation effects */
-	FTransform InterpolationSnapshotTransform;
-};
+	/** Gets the event for hovering update which is broadcasted for each interactor */
+	DECLARE_EVENT_ThreeParams( UViewportWorldInteraction, FOnVIHoverUpdate, class UViewportInteractor* /* Interactor */, FVector& /* OutHoverImpactPoint */, bool& /* bWasHandled */ );
+	FOnVIHoverUpdate& OnViewportInteractionHoverUpdate() { return OnHoverUpdateEvent; }
 
-UCLASS()
-class VIEWPORTINTERACTION_API UViewportWorldInteraction : public UObject, public IViewportWorldInteractionInterface
-{
-	GENERATED_UCLASS_BODY()
-	
-public:
+	/** Gets the event for input actions update which is broadcasted for each interactor */
+	DECLARE_EVENT_FiveParams( UViewportWorldInteraction, FOnVIActionHandle, class FEditorViewportClient& /* ViewportClient */, class UViewportInteractor* /* Interactor */, const struct FViewportActionKeyInput& /* Action */, bool& /* bOutIsInputCaptured */, bool& /* bWasHandled */ );
+	FOnVIActionHandle& OnViewportInteractionInputAction() { return OnInputActionEvent; }
 
-	virtual ~UViewportWorldInteraction();
+	/** Called when the user clicks on 'nothing' to allow systems to deselect anything that's selected */
+	DECLARE_EVENT_ThreeParams( UViewportWorldInteraction, FOnViewportInteractionInputUnhandled, class FEditorViewportClient& /* ViewportClient */, class UViewportInteractor* /* Interactor */, const struct FViewportActionKeyInput& /* Action */ );
+	FOnViewportInteractionInputUnhandled& OnViewportInteractionInputUnhandled() { return OnViewportInteractionInputUnhandledEvent; };
 
-	// IViewportWorldInteraction overrides
-	virtual void Init( UWorld* InWorld ) override;
-	virtual void Shutdown() override;
-	virtual void Tick( const float DeltaTime ) override;
-	virtual void AddInteractor( UViewportInteractor* Interactor ) override;
-	virtual void RemoveInteractor( UViewportInteractor* Interactor ) override;
-	virtual FOnVIHoverUpdate& OnViewportInteractionHoverUpdate() override { return OnHoverUpdateEvent; }
-	virtual FOnVIActionHandle& OnViewportInteractionInputAction() override { return OnInputActionEvent; }
-	virtual FOnHandleInputKey& OnHandleKeyInput() override { return OnKeyInputEvent; }
-	virtual FOnHandleInputAxis& OnHandleAxisInput() override { return OnAxisInputEvent; }
-
-	/** Sets the Viewport to be used, this is needed to run the ViewportWorldInteraction */
-	void SetViewport( const TSharedPtr<class FEditorViewportClient>& InEditorViewportClient );
-	
-	/** Starts if bInActive is true and closes if bInActivate is false */
-	void Activate( const bool bInActivate );
-	
-	/** Returns true if this worldinteraction is currently running */
-	bool IsActive() const;
+	/** To handle raw key input from the Inputprocessor */
+	DECLARE_EVENT_FourParams( UViewportWorldInteraction, FOnHandleInputKey, class FEditorViewportClient* /* ViewportClient */, const FKey /* Key */, const EInputEvent /* Event */, bool& /* bWasHandled */ );
+	FOnHandleInputKey& OnHandleKeyInput() { return OnKeyInputEvent; }
+
+	/** To handle raw axis input from the Inputprocessor */
+	DECLARE_EVENT_SixParams( UViewportWorldInteraction, FOnHandleInputAxis, class FEditorViewportClient* /* ViewportClient */, const int32 /* ControllerId */, const FKey /* Key */, const float /* Delta */, const float /* DeltaTime */, bool& /* bWasHandled */ );
+	FOnHandleInputAxis& OnHandleAxisInput() { return OnAxisInputEvent; }
 
+	/** Gets the event for when an interactor starts dragging */
+	DECLARE_EVENT_OneParam( UViewportWorldInteraction, FOnStartDragging, class UViewportInteractor* /** Interactor */ );
+	FOnStartDragging& OnStartDragging() { return OnStartDraggingEvent; };
+
+	/** Gets the event for when an interactor stops dragging */
+	DECLARE_EVENT_OneParam( UViewportWorldInteraction, FOnStopDragging, class UViewportInteractor* /** Interactor */ );
+	FOnStopDragging& OnStopDragging() { return OnStopDraggingEvent; };
+
+	/** Gets the event for when the current set of transformables has finally stop moving (after a drag, and all inertial effects or snapping interpolation has completed */
+	DECLARE_EVENT( UViewportWorldInteraction, FOnFinishedMovingTransformables );
+	FOnFinishedMovingTransformables& OnFinishedMovingTransformables() { return OnFinishedMovingTransformablesEvent; };
+
+	DECLARE_EVENT_OneParam( UViewportWorldInteraction, FOnViewportWorldInteractionTick, const float /* DeltaTime */ );
+	FOnViewportWorldInteractionTick& OnPreWorldInteractionTick() { return OnPreWorldInteractionTickEvent; };
+	FOnViewportWorldInteractionTick& OnPostWorldInteractionTick() { return OnPostWorldInteractionTickEvent; };
+
+	/** Sets the system that should be used to transform objects in the scene.  Only one can be active at a time.  If you pass
+	    nullptr, a default transformer for actors/components will be enabled */
+	void SetTransformer( class UViewportTransformer* NewTransformer );
+
+	/** Sets the list of objects that this system will be responsible for transforming when interacting using the
+	    gizmo or directly on the objects */
+	void SetTransformables( TArray< TUniquePtr< FViewportTransformable > >&& NewTransformables );
+
+	/** When using VR, this sets the viewport client that's been "possessed" by the head mounted display.  Only valid when VR is enabled. */
+	void SetDefaultOptionalViewportClient( const TSharedPtr<class FEditorViewportClient>& InEditorViewportClient );
+	
 	/** Pairs to interactors by setting	the other interactor for each interactor */
 	void PairInteractors( UViewportInteractor* FirstInteractor, UViewportInteractor* SecondInteractor );
 	
+
 	//
 	// Input
  	//
 
+	virtual bool InputKey( FEditorViewportClient* InViewportClient, FViewport* Viewport, FKey Key, EInputEvent Event ) override;
+	virtual bool InputAxis( FEditorViewportClient* InViewportClient, FViewport* Viewport, int32 ControllerId, FKey Key, float Delta, float DeltaTime ) override;
+
 	/** Handles the key input from the InputProcessor*/
-	bool HandleInputKey( const FKey Key, const EInputEvent Event );
+	bool PreprocessedInputKey( const FKey Key, const EInputEvent Event );
 
 	/** Handles the axis input from the InputProcessor*/
-	bool HandleInputAxis( const int32 ControllerId, const FKey Key, const float Delta, const float DeltaTime );
+	bool PreprocessedInputAxis( const int32 ControllerId, const FKey Key, const float Delta, const float DeltaTime );
 
 	/** Gets the world space transform of the calibrated VR room origin.  When using a seated VR device, this will feel like the
 	camera's world transform (before any HMD positional or rotation adjustments are applied.) */
@@ -127,6 +143,9 @@ public:
 	/** Gets the transform of the viewport / user's HMD in world space */
 	FTransform GetHeadTransform() const;
 
+	/** Returns true if we actually are using VR and have a valid head location, meaning a call to GetHeadTransform() is valid */
+	bool HaveHeadTransform() const;
+
 	/** Sets a new transform for the room, in world space.  This is basically setting the editor's camera transform for the viewport */
 	void SetRoomTransform( const FTransform& NewRoomTransform );
 	
@@ -134,10 +153,10 @@ public:
 	float GetWorldScaleFactor() const;
 
 	/** Gets the world of the viewport used for this worldinteraction */
-	UWorld* GetViewportWorld() const;
+	virtual UWorld* GetWorld() const;
 
 	/** Gets the currently used viewport */
-	FEditorViewportClient* GetViewportClient() const;
+	FEditorViewportClient* GetDefaultOptionalViewportClient() const;
 
 	/** Returns the Unreal controller ID for the motion controllers we're using */
 	int32 GetMotionControllerID() const
@@ -145,9 +164,6 @@ public:
 		return MotionControllerID;
 	}
 
-	/** Sets that the worldinteractor is used in VR, to give a correct head transform */
-	void SetIsUsingVR( const bool bInIsUsingVR );
-
 	/** Invokes the editor's undo system to undo the last change */
 	void Undo();
 
@@ -172,12 +188,11 @@ public:
 	/** Called to finish a drag action with the specified interactor */
 	void StopDragging( class UViewportInteractor* Interactor );
 
-	virtual FOnStopDragging& OnStopDragging() override { return OnStopDraggingEvent; };
-
-	virtual FOnWorldScaleChanged& OnWorldScaleChanged() override { return OnWorldScaleChangedEvent; };
+	/** Starts dragging selected objects around.  Called when clicking and dragging on actors/gizmos in the world, or when placing new objects.  Returns true if dragging actually started. */
+	bool StartDragging( UViewportInteractor* Interactor, UActorComponent* ClickedTransformGizmoComponent, const FVector& HitLocation, const bool bIsPlacingNewObjects, const bool bStartTransaction );
 
-	/** Starts dragging selected actors around, changing selection if needed.  Called when clicking in the world, or when placing new objects */
-	void StartDraggingActors( UViewportInteractor* Interactor, const FViewportActionKeyInput& Action, UActorComponent* ClickedComponent, const FVector& HitLocation, const bool bIsPlacingActors );
+	DECLARE_EVENT_OneParam( UViewportWorldInteraction, FOnWorldScaleChanged, const float /* NewWorldToMetersScale */);
+	virtual FOnWorldScaleChanged& OnWorldScaleChanged() { return OnWorldScaleChangedEvent; };
 	
 	/** Sets which transform gizmo coordinate space is used */
 	void SetTransformGizmoCoordinateSpace( const ECoordSystem NewCoordSystem );
@@ -202,9 +217,6 @@ public:
 		bSkipInteractiveWorldMovementThisFrame = true;
 	}
 
-	/** True if we're actively interacting with the specified actor */
-	bool IsTransformingActor( AActor* Actor ) const;
-
 	/** Checks to see if we're allowed to interact with the specified component */
 	bool IsInteractableComponent( const UActorComponent* Component ) const;
 
@@ -214,9 +226,19 @@ public:
 	/** Sets whether the transform gizmo should be visible at all */
 	void SetTransformGizmoVisible( const bool bShouldBeVisible );
 
-	/** Returns whether the transform gizmo is visible */
+	/** Returns whether the transform gizmo should be visible */
+	bool ShouldTransformGizmoBeVisible() const;
+
+	/** Returns whether the transform gizmo is actually visible right now.  This is different than whether it 'should' be
+	    visible, as the gizmo can be implicitly hidden if there are no selected objects, for example. */
 	bool IsTransformGizmoVisible() const;
 
+	/** Sets how large the transform gizmo is */
+	void SetTransformGizmoScale( const float NewScale );
+
+	/** Returns the current transform gizmo scale */
+	float GetTransformGizmoScale() const;
+
 	/** Sets if objects are dragged with either hand since last time selecting something */
 	void SetDraggedSinceLastSelection( const bool bNewDraggedSinceLastSelection );
 
@@ -247,40 +269,41 @@ public:
 	/** Switch which transform gizmo coordinate space we're using. */
 	void CycleTransformGizmoCoordinateSpace();
 
-	/** Destroys a transient actor we created earlier and nulls out the pointer */
-	void DestroyTransientActor(AActor* Actor) const;
-	
 	/** Figures out where to place an object when tracing it against the scene using a laser pointer */
 	bool FindPlacementPointUnderLaser( UViewportInteractor* Interactor, FVector& OutHitLocation );
 
 	/** Gets the tracking transactions */
 	FTrackingTransaction& GetTrackingTransaction();
 	
-	/** Creates a Transformable for each selected actor, wiping out the existing transformables */
-	void SetupTransformablesForSelectedActors();
-	
-	/** All of the objects we're currently interacting with, such as selected actors */
-	TArray< TUniquePtr< class FViewportTransformable > > Transformables;
+	void SetUseInputProcessor( bool bInUseInputProcessor );
+
+	/** Resets the start transform all of transformables to match current transform */
+	void ResetTransformables();
 
-	/** Spawns a transient actor that we can use in the editor world (templated for convenience) */
+	/** Gets a list of all of the objects we're currently interacting with, such as selected actors */
+	TArray< TUniquePtr< FViewportTransformable > >& GetTransformables()
+	{
+		return Transformables;
+	}
+
+	/** Static: Spawns a transient actor that we can use in the editor world (templated for convenience) */
 	template<class T>
-	T* SpawnTransientSceneActor(const FString& ActorName, const bool bWithSceneComponent) const
+	static inline T* SpawnTransientSceneActor(UWorld* World, const FString& ActorName, const bool bWithSceneComponent)
 	{
-		return CastChecked<T>(SpawnTransientSceneActor(T::StaticClass(), ActorName, bWithSceneComponent));
+		return CastChecked<T>(SpawnTransientSceneActor(World, T::StaticClass(), ActorName, bWithSceneComponent));
 	}
 
-protected:
+	/** Static: Spawns a transient actor that we can use in the editor world */
+	static AActor* SpawnTransientSceneActor( UWorld* World, TSubclassOf<AActor> ActorClass, const FString& ActorName, const bool bWithSceneComponent );
 
-	/** Allow base class to refresh on selection changed */
-	virtual void RefreshOnSelectionChanged( const bool bNewObjectsSelected, bool bAllHandlesVisible, class UActorComponent* SingleVisibleHandle, const TArray< UActorComponent* >& HoveringOverHandles );
+	/** Static: Destroys a transient actor we created earlier */
+	static void DestroyTransientActor( UWorld* World, AActor* Actor );
 
-private:
 
-	/** Starts when activated */
-	void Enter();
+protected:
+
 
-	/** Closes when deactivated */
-	void Exit();
+private:
 
 	/** Called when the editor is closed */
 	void OnEditorClosed();
@@ -309,6 +332,7 @@ private:
 		const FVector& OtherHandDragDeltaFromStart, 
 		const FVector& LaserPointerStart, 
 		const FVector& LaserPointerDirection, 
+		const float LaserPointerMaxLength,
 		const bool bIsLaserPointerValid, 
 		const FTransform& GizmoStartTransform, 
 		const FBox& GizmoStartLocalBounds, 
@@ -328,10 +352,15 @@ private:
 		const FVector& LaserPointerDirection, 
 		const bool bIsLaserPointerValid, 
 		const FTransform& GizmoStartTransform, 
+		const float LaserPointerMaxLength, 
 		FVector& GizmoSpaceFirstDragUpdateOffsetAlongAxis, 
 		FVector& DragDeltaFromStartOffset,
 		FVector& OutClosestPointOnLaser) const;
 
+	/** Called after dragging has finished with one or both hands and all inertial movement and interpolation has finished with
+	    the set of transformables we were working with.  Will also be called if the drag was interrupted. */
+	void FinishedMovingTransformables();
+
 	/** Checks to see if smooth snapping is currently enabled */
 	bool IsSmoothSnappingEnabled() const;
 
@@ -341,15 +370,9 @@ private:
 		return FTimespan::FromSeconds( FApp::GetCurrentTime() ) - AppTimeEntered;
 	}
 
-	/** Updates the transforma */
-	void UpdateTransformables( const float DeltaTime );
-
 	/** Conditionally polls input from interactors, if we haven't done that already this frame */
 	void PollInputIfNeeded();
 
-	/** Called when selection changes in the level */
-	void OnActorSelectionChanged( UObject* ChangedObject );
-
 	/** Refreshes the transform gizmo to match the currently selected actor set */
 	void RefreshTransformGizmo( const bool bNewObjectsSelected, bool bAllHandlesVisible, class UActorComponent* SingleVisibleHandle, const TArray< UActorComponent* >& HoveringOverHandles );
 
@@ -359,12 +382,10 @@ private:
 	/** Gets the inertie contribution of another interactor */
 	UViewportInteractor* GetOtherInteractorIntertiaContribute( UViewportInteractor* Interactor );
 
-	/** Spawns a transient actor that we can use in the editor world */
-	AActor* SpawnTransientSceneActor( TSubclassOf<AActor> ActorClass, const FString& ActorName, const bool bWithSceneComponent ) const;
-
 	/** Destroys the actors */
 	void DestroyActors();
 
+public:
 	/** Gets the snap grid mesh actor */
 	AActor* GetSnapGridActor()
 	{
@@ -372,6 +393,7 @@ private:
 		return SnapGridActor;
 	}
 
+private:
 	/** Gets the snap grid mesh MID */
 	class UMaterialInstanceDynamic* GetSnapGridMID()
 	{
@@ -409,9 +431,6 @@ private:
 	TArray<FLinearColor> Colors;
 
 protected:
-	
-	/** The world for this interaction */
-	UWorld* World;
 
 	/** True if we've dragged objects with either hand since the last time we selected something */
 	bool bDraggedSinceLastSelection;
@@ -432,14 +451,19 @@ private:
 	FTimespan AppTimeEntered;
 
 	/** All the interactors registered to modify the world */
+	UPROPERTY()
 	TArray< class UViewportInteractor* > Interactors;
 
+	/** The active system being used to transform objects */
+	UPROPERTY()
+	class UViewportTransformer* ViewportTransformer;
+
 	//
 	// Viewport
 	//
 
 	/** The viewport that the worldinteraction is used for */
-	class FEditorViewportClient* EditorViewportClient;
+	class FEditorViewportClient* DefaultOptionalViewportClient;
 
 	/** The last frame that input was polled.  We keep track of this so that we can make sure to poll for latest input right before
 	its first needed each frame */
@@ -471,6 +495,10 @@ private:
 	// Object movement/placement
 	//
 
+	/** True if transformables are being moved around, including any inertial movement or movement caused by
+	    smooth snapping or placement interpolation. */
+	bool bAreTransformablesMoving;
+
 	/** True if our transformables are currently interpolating to their target location */
 	bool bIsInterpolatingTransformablesFromSnapshotTransform;
 
@@ -500,11 +528,14 @@ private:
 	/** The direction of where the rotation handle is facing when starting rotation */
 	TOptional<FVector> StartDragHandleDirection;
 
-	/** The current Gizmo type that is used for the TransformGizmo Actor */
-	EGizmoHandleTypes CurrentGizmoType;
-
 	/** Whether the gizmo should be visible or not */
-	bool bIsTransformGizmoVisible;
+	bool bShouldTransformGizmoBeVisible;
+
+	/** How big the transform gizmo should be */
+	float TransformGizmoScale;
+
+	/** All of the objects we're currently interacting with, such as selected actors */
+	TArray< TUniquePtr< FViewportTransformable > > Transformables;
 
 	//
 	// Snap grid
@@ -539,18 +570,53 @@ private:
 	/** Event that is fired for when a key is pressed by an interactor */
 	FOnVIActionHandle OnInputActionEvent;
 
+	/** Called when the user clicks on 'nothing' to request systems to deselect anything that's selected */
+	FOnViewportInteractionInputUnhandled OnViewportInteractionInputUnhandledEvent;
+
 	/** Event that is fired when new key input is received by the InputProcessor */
 	FOnHandleInputKey OnKeyInputEvent;
 
 	/** Event that is fired when new axis input is received by the InputProcessor */
 	FOnHandleInputAxis OnAxisInputEvent;
 
-	/** Event that is fired when a interactor stopped dragging */
+	/** Event that is fired when an interactor starts dragging */
+	FOnStartDragging OnStartDraggingEvent;
+
+	/** Event that is fired when an interactor stopped dragging */
 	FOnStopDragging OnStopDraggingEvent;
+	
+	/** Event for when the current set of transformables has finally stop moving (after a drag, and all inertial effects or snapping interpolation has completed */
+	FOnFinishedMovingTransformables OnFinishedMovingTransformablesEvent;
+
+	/** Event to tick before the worldinteraction ticks */
+	FOnViewportWorldInteractionTick OnPreWorldInteractionTickEvent;
+
+	/** Event to tick after the worldinteraction has ticked */
+	FOnViewportWorldInteractionTick OnPostWorldInteractionTickEvent;
+
+
+	//
+	// Interactors
+	//
+
+	/** The default mouse cursor interactor, activated on demand */
+	UPROPERTY()
+	class UMouseCursorInteractor* DefaultMouseCursorInteractor;
+
+	/** Reference count for the default mouse cursor interactor.  When this reaches zero, the mouse cursor interactor goes away. */
+	int32 DefaultMouseCursorInteractorRefCount;
+
+
+	//
+	// Configuration
+	//
 
 	/** Event that is fired when the world scale changes */
 	FOnWorldScaleChanged OnWorldScaleChangedEvent;
 
 	/** If the world interaction is active and running the tick function */
 	bool bActive;
+
+	/** If this world interaction should get input from the input processor */
+	bool bUseInputProcessor;
 };
diff --git a/Engine/Source/Editor/ViewportInteraction/Public/ViewportWorldInteractionInterface.h b/Engine/Source/Editor/ViewportInteraction/Public/ViewportWorldInteractionInterface.h
deleted file mode 100644
diff --git a/Engine/Source/Editor/ViewportInteraction/ViewportInteraction.Build.cs b/Engine/Source/Editor/ViewportInteraction/ViewportInteraction.Build.cs
--- a/Engine/Source/Editor/ViewportInteraction/ViewportInteraction.Build.cs
+++ b/Engine/Source/Editor/ViewportInteraction/ViewportInteraction.Build.cs
@@ -8,7 +8,7 @@ namespace UnrealBuildTool.Rules
     {
         public ViewportInteraction(TargetInfo Target)
         {
-            PrivateDependencyModuleNames.AddRange(
+            PublicDependencyModuleNames.AddRange(
                 new string[] {
                     "Core",
                     "CoreUObject",
@@ -17,7 +17,6 @@ namespace UnrealBuildTool.Rules
                     "UnrealEd",
                     "Slate",
                     "SlateCore",
-                    "HeadMountedDisplay"
                 }
             );
 
@@ -26,11 +25,6 @@ namespace UnrealBuildTool.Rules
 					"LevelEditor"
                 }
             );
-
-            DynamicallyLoadedModuleNames.AddRange(
-                new string[] {
-                }
-            );
-        }
-    }
+		}
+	}
 }
\ No newline at end of file
diff --git a/Engine/Source/Editor/ViewportInteraction/ViewportInteractionInputProcessor.cpp b/Engine/Source/Editor/ViewportInteraction/ViewportInteractionInputProcessor.cpp
--- a/Engine/Source/Editor/ViewportInteraction/ViewportInteractionInputProcessor.cpp
+++ b/Engine/Source/Editor/ViewportInteraction/ViewportInteractionInputProcessor.cpp
@@ -21,20 +21,21 @@ void FViewportInteractionInputProcessor::Tick( const float DeltaTime, FSlateAppl
 
 bool FViewportInteractionInputProcessor::HandleKeyDownEvent( FSlateApplication& SlateApp, const FKeyEvent& InKeyEvent )
 {
-	return WorldInteractionManager->HandleInputKey( InKeyEvent.GetKey(), InKeyEvent.IsRepeat() ? IE_Repeat : IE_Pressed );
+	return WorldInteractionManager->PreprocessedInputKey( InKeyEvent.GetKey(), InKeyEvent.IsRepeat() ? IE_Repeat : IE_Pressed );
 }
 
 bool FViewportInteractionInputProcessor::HandleKeyUpEvent( FSlateApplication& SlateApp, const FKeyEvent& InKeyEvent )
 {
-	return WorldInteractionManager->HandleInputKey( InKeyEvent.GetKey(), IE_Released );
+	return WorldInteractionManager->PreprocessedInputKey( InKeyEvent.GetKey(), IE_Released );
 }
 
 bool FViewportInteractionInputProcessor::HandleAnalogInputEvent( FSlateApplication& SlateApp, const FAnalogInputEvent& InAnalogInputEvent )
 {
-	return WorldInteractionManager->HandleInputAxis( InAnalogInputEvent.GetUserIndex(), InAnalogInputEvent.GetKey(), InAnalogInputEvent.GetAnalogValue(), FApp::GetDeltaTime() );
+	return WorldInteractionManager->PreprocessedInputAxis( InAnalogInputEvent.GetUserIndex(), InAnalogInputEvent.GetKey(), InAnalogInputEvent.GetAnalogValue(), FApp::GetDeltaTime() );
 }
 
 bool FViewportInteractionInputProcessor::HandleMouseMoveEvent( FSlateApplication& SlateApp, const FPointerEvent& MouseEvent )
 {
 	return false;
 }
+
diff --git a/Engine/Source/Editor/ViewportInteraction/ViewportInteractionInputProcessor.h b/Engine/Source/Editor/ViewportInteraction/ViewportInteractionInputProcessor.h
--- a/Engine/Source/Editor/ViewportInteraction/ViewportInteractionInputProcessor.h
+++ b/Engine/Source/Editor/ViewportInteraction/ViewportInteractionInputProcessor.h
@@ -18,18 +18,11 @@ public:
 	FViewportInteractionInputProcessor ( FViewportWorldInteractionManager* InWorldInteractionManager );
 	virtual ~FViewportInteractionInputProcessor ();
 
+	// IInputProcess overrides
 	virtual void Tick( const float DeltaTime, FSlateApplication& SlateApp, TSharedRef<ICursor> Cursor ) override;
-
-	/** Key down input */
 	virtual bool HandleKeyDownEvent( FSlateApplication& SlateApp, const FKeyEvent& InKeyEvent ) override;
-	
-	/** Key up input */
 	virtual bool HandleKeyUpEvent( FSlateApplication& SlateApp, const FKeyEvent& InKeyEvent ) override;
-	
-	/** Analog axis input */
 	virtual bool HandleAnalogInputEvent( FSlateApplication& SlateApp, const FAnalogInputEvent& InAnalogInputEvent ) override;
-
-	/** Mouse movement input */
 	virtual bool HandleMouseMoveEvent( FSlateApplication& SlateApp, const FPointerEvent& MouseEvent ) override;
 
 private:
diff --git a/Engine/Source/Editor/ViewportInteraction/ViewportInteractionModule.cpp b/Engine/Source/Editor/ViewportInteraction/ViewportInteractionModule.cpp
--- a/Engine/Source/Editor/ViewportInteraction/ViewportInteractionModule.cpp
+++ b/Engine/Source/Editor/ViewportInteraction/ViewportInteractionModule.cpp
@@ -1,6 +1,7 @@
 // Copyright 1998-2017 Epic Games, Inc. All Rights Reserved.
 
 #include "ViewportInteractionModule.h"
+#include "ViewportWorldInteractionManager.h"
 
 FViewportInteractionModule::FViewportInteractionModule()
 {
@@ -28,8 +29,3 @@ FViewportWorldInteractionManager& FViewportInteractionModule::GetWorldInteractio
 {
 	return WorldInteractionManager;
 }
-
-void FViewportInteractionModule::Tick( float DeltaTime )
-{
-	WorldInteractionManager.Tick( DeltaTime );
-}
diff --git a/Engine/Source/Editor/ViewportInteraction/ViewportInteractionModule.h b/Engine/Source/Editor/ViewportInteraction/ViewportInteractionModule.h
--- a/Engine/Source/Editor/ViewportInteraction/ViewportInteractionModule.h
+++ b/Engine/Source/Editor/ViewportInteraction/ViewportInteractionModule.h
@@ -22,7 +22,6 @@ public:
 		return true;
 	}
 	virtual FViewportWorldInteractionManager& GetWorldInteractionManager() override;
-	virtual void Tick( float DeltaTime ) override;
 
 private:
 
diff --git a/Engine/Source/Editor/ViewportInteraction/ViewportInteractor.cpp b/Engine/Source/Editor/ViewportInteraction/ViewportInteractor.cpp
--- a/Engine/Source/Editor/ViewportInteraction/ViewportInteractor.cpp
+++ b/Engine/Source/Editor/ViewportInteraction/ViewportInteractor.cpp
@@ -6,9 +6,12 @@
 #include "Editor.h"
 #include "ViewportWorldInteraction.h"
 #include "ViewportInteractableInterface.h"
+#include "MouseCursorInteractor.h"
 
 namespace VI
 {
+	static FAutoConsoleVariable GrabberSphereRadius( TEXT( "VI.GrabberSphereRadius" ), 15.0f, TEXT( "In radial mode, the radius of the sphere used to select and interact" ) );
+	static FAutoConsoleVariable GrabberSphereOffset( TEXT( "VI.GrabberSphereOffset" ), 10.0f, TEXT( "Offset from the controller origin that the grabber sphere should be centered at" ) );
 	static FAutoConsoleVariable LaserPointerMaxLength( TEXT( "VI.LaserPointerMaxLength" ), 10000.0f, TEXT( "Maximum length of the laser pointer line" ) );
 	static FAutoConsoleVariable DragHapticFeedbackStrength( TEXT( "VI.DragHapticFeedbackStrength" ), 1.0f, TEXT( "Default strength for haptic feedback when starting to drag objects" ) ); //@todo ViewportInteraction: Duplicate from ViewportWorldInteraction
 	static FAutoConsoleVariable SelectionHapticFeedbackStrength( TEXT( "VI.SelectionHapticFeedbackStrength" ), 0.5f, TEXT( "Default strength for haptic feedback when selecting objects" ) );
@@ -19,9 +22,9 @@ UViewportInteractor::UViewportInteractor( const FObjectInitializer& Initializer
 	InteractorData(),
 	KeyToActionMap(),
 	WorldInteraction( nullptr ),
-	OtherInteractor( nullptr )
+	OtherInteractor( nullptr ),
+	bAllowGrabberSphere( true )
 {
-
 }
 
 FViewportInteractorData& UViewportInteractor::GetInteractorData()
@@ -80,7 +83,7 @@ void UViewportInteractor::RemoveKeyAction( const FKey& Key )
 	KeyToActionMap.Remove( Key );
 }
 
-bool UViewportInteractor::HandleInputKey( const FKey Key, const EInputEvent Event )
+bool UViewportInteractor::HandleInputKey( FEditorViewportClient& ViewportClient, const FKey Key, const EInputEvent Event )
 {
 	bool bHandled = false;
 	FViewportActionKeyInput* Action = KeyToActionMap.Find( Key );
@@ -88,14 +91,26 @@ bool UViewportInteractor::HandleInputKey( const FKey Key, const EInputEvent Even
 	{
 		Action->Event = Event;
 		
-		// Give subsystems a chance to handle actions for this interactor
-		WorldInteraction->OnViewportInteractionInputAction().Broadcast( *WorldInteraction->GetViewportClient(), this, *Action, Action->bIsInputCaptured, bHandled );
+		// Prefer transform gizmo interactions over anything else
+		FHitResult HitResult = GetHitResultFromLaserPointer();
+		const bool bPressedTransformGizmo =
+			Event == IE_Pressed &&
+			HitResult.Actor.IsValid() &&
+			HitResult.Actor == WorldInteraction->GetTransformGizmoActor() &&
+			( Action->ActionType == ViewportWorldActionTypes::SelectAndMove ||
+			  Action->ActionType == ViewportWorldActionTypes::SelectAndMove_LightlyPressed );
+		if( !bPressedTransformGizmo )
+		{
+			// Give subsystems a chance to handle actions for this interactor
+			WorldInteraction->OnViewportInteractionInputAction().Broadcast( ViewportClient, this, *Action, Action->bIsInputCaptured, bHandled );
+		}
 
 		if(!bHandled)
 		{
 			// Give the derived classes a chance to update according to the input
-			HandleInputKey( *Action, Key, Event, bHandled );
+			HandleInputKey( ViewportClient, *Action, Key, Event, bHandled );
 		}
+
 		// Start checking on default action implementation
 		if ( !bHandled )
 		{
@@ -112,12 +127,11 @@ bool UViewportInteractor::HandleInputKey( const FKey Key, const EInputEvent Even
 				if ( Event == IE_Pressed )
 				{
 					// No clicking while we're dragging the world.  (No laser pointers are visible, anyway.)
-					FHitResult HitResult = GetHitResultFromLaserPointer();
 					if ( !bIsDraggingWorldWithTwoHands && !bHandled && HitResult.Actor.IsValid() )
 					{
 						if ( WorldInteraction->IsInteractableComponent( HitResult.GetComponent() ) )
 						{
-							InteractorData.ClickingOnComponent = HitResult.GetComponent();
+							InteractorData.ClickingOnComponent = HitResult.GetComponent();	// @todo gizmo: Should be changed to store only gizmo components?
 
 							AActor* Actor = HitResult.Actor.Get();
 
@@ -153,20 +167,24 @@ bool UViewportInteractor::HandleInputKey( const FKey Key, const EInputEvent Even
 
 								// Is the other hand already dragging this stuff?
 								if ( OtherInteractorData != nullptr &&
-									 ( OtherInteractorData->DraggingMode == EViewportInteractionDraggingMode::ActorsWithGizmo ||
-									   OtherInteractorData->DraggingMode == EViewportInteractionDraggingMode::ActorsFreely ) )
+									 ( OtherInteractorData->DraggingMode == EViewportInteractionDraggingMode::TransformablesWithGizmo ||
+									   OtherInteractorData->DraggingMode == EViewportInteractionDraggingMode::TransformablesFreely ) )
 								{
-									// Only if they clicked on one of the objects we're moving
-									if ( WorldInteraction->IsTransformingActor( Actor ) )
+									// Only if they clicked on one of the objects we're already moving
+									if ( Actor->IsSelected() )	// @todo gizmo: Should check for existing transformable, not an actor!
 									{
 										// If we were dragging with a gizmo, we'll need to stop doing that and instead drag freely.
-										if ( OtherInteractorData->DraggingMode == EViewportInteractionDraggingMode::ActorsWithGizmo )
+										if ( OtherInteractorData->DraggingMode == EViewportInteractionDraggingMode::TransformablesWithGizmo )
 										{
 											WorldInteraction->StopDragging( this );
 
-											FViewportActionKeyInput OtherInteractorAction( ViewportWorldActionTypes::SelectAndMove );
-											const bool bIsPlacingActors = false;
-											WorldInteraction->StartDraggingActors( OtherInteractor, OtherInteractorAction, HitResult.GetComponent(), OtherInteractorData->HoverLocation, bIsPlacingActors ); //@todo ViewportInteraction
+											UPrimitiveComponent* ClickedTransformGizmoComponent = nullptr;
+											const bool bIsPlacingNewObjects = false;
+											const bool bStartTransaction = true;
+											if( WorldInteraction->StartDragging( OtherInteractor, ClickedTransformGizmoComponent, OtherInteractorData->HoverLocation, bIsPlacingNewObjects, bStartTransaction ) )
+											{
+												// ...
+											}
 										}
 
 										InteractorData.DraggingMode = InteractorData.LastDraggingMode = EViewportInteractionDraggingMode::AssistingDrag;
@@ -201,11 +219,11 @@ bool UViewportInteractor::HandleInputKey( const FKey Key, const EInputEvent Even
 										// @todo vreditor: We don't currently support selection/dragging separate objects with either hand
 									}
 								}
-								else if ( OtherInteractorData != nullptr && OtherInteractorData->DraggingMode == EViewportInteractionDraggingMode::ActorsWithGizmo )
+								else if ( OtherInteractorData != nullptr && OtherInteractorData->DraggingMode == EViewportInteractionDraggingMode::TransformablesWithGizmo )
 								{
 									// We don't support dragging objects with the gizmo using two hands.  Just absorb it.
 								}
-								else if ( OtherInteractorData != nullptr && OtherInteractorData->DraggingMode == EViewportInteractionDraggingMode::ActorsAtLaserImpact )
+								else if ( OtherInteractorData != nullptr && OtherInteractorData->DraggingMode == EViewportInteractionDraggingMode::TransformablesAtLaserImpact )
 								{
 									// Doesn't work with two hands.  Just absorb it.
 								}
@@ -283,8 +301,27 @@ bool UViewportInteractor::HandleInputKey( const FKey Key, const EInputEvent Even
 
 									if ( bShouldDragSelected && InteractorData.DraggingMode != EViewportInteractionDraggingMode::Interactable )
 									{
-										const bool bIsPlacingActors = false;
-										WorldInteraction->StartDraggingActors( this, *Action, HitResult.GetComponent(), HitResult.ImpactPoint, bIsPlacingActors ); //@todo ViewportInteraction
+										UPrimitiveComponent* ClickedTransformGizmoComponent = nullptr;
+										{
+											const bool bIsMouseCursorInteractor = this->IsA( UMouseCursorInteractor::StaticClass() );
+
+											const bool bUsingGizmo =
+												HitResult.GetComponent() != nullptr &&
+												HitResult.GetComponent()->GetOwner() == WorldInteraction->GetTransformGizmoActor() &&
+												( Action->ActionType == ViewportWorldActionTypes::SelectAndMove_LightlyPressed ||
+												( bIsMouseCursorInteractor && Action->ActionType == ViewportWorldActionTypes::SelectAndMove ) );		// Only use the gizmo when lightly pressed (unless we're dealing with the mouse cursor)
+											if( bUsingGizmo )
+											{
+												ClickedTransformGizmoComponent = HitResult.GetComponent();
+											}
+										}
+
+										const bool bIsPlacingNewObjects = false;
+										const bool bStartTransaction = true;
+										if( WorldInteraction->StartDragging( this, ClickedTransformGizmoComponent, HitResult.ImpactPoint, bIsPlacingNewObjects, bStartTransaction ) )
+										{
+											Action->bIsInputCaptured = true;
+										}
 									}
 									else if ( bSelectionChanged )
 									{
@@ -401,13 +438,15 @@ bool UViewportInteractor::HandleInputKey( const FKey Key, const EInputEvent Even
 			{
 				WorldInteraction->Deselect();
 			}
+
+			WorldInteraction->OnViewportInteractionInputUnhandled().Broadcast( ViewportClient, this, *Action );
 		}
 	}
 
 	return bHandled;
 }
 
-bool UViewportInteractor::HandleInputAxis( const FKey Key, const float Delta, const float DeltaTime )
+bool UViewportInteractor::HandleInputAxis( FEditorViewportClient& ViewportClient, const FKey Key, const float Delta, const float DeltaTime )
 {
 	bool bHandled = false;
 	FViewportActionKeyInput* KnownAction = KeyToActionMap.Find( Key );
@@ -416,7 +455,7 @@ bool UViewportInteractor::HandleInputAxis( const FKey Key, const float Delta, co
 		FViewportActionKeyInput Action( KnownAction->ActionType );
 
 		// Give the derived classes a chance to update according to the input
-		HandleInputAxis( Action, Key, Delta, DeltaTime, bHandled );
+		HandleInputAxis( ViewportClient, Action, Key, Delta, DeltaTime, bHandled );
 	}
 
 	return bHandled;
@@ -442,7 +481,7 @@ FVector UViewportInteractor::GetDragTranslationVelocity() const
 	return InteractorData.DragTranslationVelocity;
 }
 
-bool UViewportInteractor::GetLaserPointer( FVector& LaserPointerStart, FVector& LaserPointerEnd, const bool bEvenIfBlocked, const float LaserLengthOverride )
+bool UViewportInteractor::GetLaserPointer( FVector& LaserPointerStart, FVector& LaserPointerEnd, const bool bEvenIfBlocked, const float LaserLengthOverride ) const
 {
 	// If we're currently grabbing the world with both hands, then the laser pointer is not available
 	if ( /*bHaveMotionController &&*/ //@todo ViewportInteraction
@@ -469,6 +508,32 @@ bool UViewportInteractor::GetLaserPointer( FVector& LaserPointerStart, FVector&
 	return false;
 }
 
+
+bool UViewportInteractor::GetGrabberSphere( FSphere& OutGrabberSphere, const bool bEvenIfBlocked ) const
+{
+	OutGrabberSphere = FSphere( 0 );
+
+	if( bAllowGrabberSphere )
+	{
+		FVector LaserPointerStart, LaserPointerEnd;
+		if( GetLaserPointer( LaserPointerStart, LaserPointerEnd, bEvenIfBlocked ) )
+		{
+			FTransform HandTransform;
+			FVector HandForwardVector;
+			if( GetTransformAndForwardVector( HandTransform, HandForwardVector ) )
+			{
+				const FVector GrabberSphereCenter = HandTransform.GetLocation() + HandForwardVector * VI::GrabberSphereOffset->GetFloat() * WorldInteraction->GetWorldScaleFactor();
+
+				OutGrabberSphere = FSphere( GrabberSphereCenter, VI::GrabberSphereRadius->GetFloat() * WorldInteraction->GetWorldScaleFactor() );
+				return true;
+			}
+		}
+	}
+
+	return false;
+}
+
+
 float UViewportInteractor::GetLaserPointerMaxLength() const
 {
 	return VI::LaserPointerMaxLength->GetFloat() * WorldInteraction->GetWorldScaleFactor();
@@ -482,9 +547,15 @@ FHitResult UViewportInteractor::GetHitResultFromLaserPointer( TArray<AActor*>* O
 	FVector LaserPointerStart, LaserPointerEnd;
 	if ( GetLaserPointer( LaserPointerStart, LaserPointerEnd, bEvenIfBlocked, LaserLengthOverride ) )
 	{
+		bool bActuallyIgnoreGizmos = bIgnoreGizmos;
+		if( !WorldInteraction->IsTransformGizmoVisible() )
+		{
+			bActuallyIgnoreGizmos = true;
+		}
+
 		// Twice twice.  Once for editor gizmos which are "on top" and always take precedence, then a second time
 		// for all of the scene objects
-		for ( int32 PassIndex = bIgnoreGizmos ? 1 : 0; PassIndex < 2; ++PassIndex )
+		for ( int32 PassIndex = bActuallyIgnoreGizmos ? 1 : 0; PassIndex < 2; ++PassIndex )
 		{
 			const bool bOnlyEditorGizmos = ( PassIndex == 0 );
 
@@ -503,7 +574,7 @@ FHitResult UViewportInteractor::GetHitResultFromLaserPointer( TArray<AActor*>* O
 				const FCollisionResponseParams& ResponseParam = FCollisionResponseParams::DefaultResponseParam;
 				const ECollisionChannel CollisionChannel = bOnlyEditorGizmos ? COLLISION_GIZMO : ECC_Visibility;
 
-				bHit = WorldInteraction->GetViewportWorld()->LineTraceSingleByChannel( HitResult, LaserPointerStart, LaserPointerEnd, CollisionChannel, TraceParams, ResponseParam );
+				bHit = WorldInteraction->GetWorld()->LineTraceSingleByChannel( HitResult, LaserPointerStart, LaserPointerEnd, CollisionChannel, TraceParams, ResponseParam );
 				if ( bHit )
 				{
 					BestHitResult = HitResult;
@@ -513,7 +584,7 @@ FHitResult UViewportInteractor::GetHitResultFromLaserPointer( TArray<AActor*>* O
 			{
 				FCollisionObjectQueryParams EverythingButGizmos( FCollisionObjectQueryParams::AllObjects );
 				EverythingButGizmos.RemoveObjectTypesToQuery( COLLISION_GIZMO );
-				bHit = WorldInteraction->GetViewportWorld()->LineTraceSingleByObjectType( HitResult, LaserPointerStart, LaserPointerEnd, EverythingButGizmos, TraceParams );
+				bHit = WorldInteraction->GetWorld()->LineTraceSingleByObjectType( HitResult, LaserPointerStart, LaserPointerEnd, EverythingButGizmos, TraceParams );
 				
 				if ( bHit )
 				{
@@ -550,7 +621,7 @@ FHitResult UViewportInteractor::GetHitResultFromLaserPointer( TArray<AActor*>* O
 	return BestHitResult;
 }
 
-bool UViewportInteractor::GetTransformAndForwardVector( FTransform& OutHandTransform, FVector& OutForwardVector )
+bool UViewportInteractor::GetTransformAndForwardVector( FTransform& OutHandTransform, FVector& OutForwardVector ) const
 {
 	OutHandTransform = InteractorData.Transform;
 	OutForwardVector = OutHandTransform.GetRotation().Vector();
@@ -616,7 +687,6 @@ bool UViewportInteractor::AllowTriggerFullPress() const
 {
 	return InteractorData.bAllowTriggerFullPress;
 }
-
 bool UViewportInteractor::IsHoveringOverPriorityType() const
 {
 	return InteractorData.bHitResultIsPriorityType;
diff --git a/Engine/Source/Editor/ViewportInteraction/ViewportTransformable.cpp b/Engine/Source/Editor/ViewportInteraction/ViewportTransformable.cpp
new file mode 100644
--- /dev/null
+++ b/Engine/Source/Editor/ViewportInteraction/ViewportTransformable.cpp
@@ -0,0 +1,5 @@
+// Copyright 1998-2016 Epic Games, Inc. All Rights Reserved.
+
+#include "ViewportTransformable.h"
+
+
diff --git a/Engine/Source/Editor/ViewportInteraction/ViewportTransformer.cpp b/Engine/Source/Editor/ViewportInteraction/ViewportTransformer.cpp
new file mode 100644
--- /dev/null
+++ b/Engine/Source/Editor/ViewportInteraction/ViewportTransformer.cpp
@@ -0,0 +1,17 @@
+// Copyright 1998-2016 Epic Games, Inc. All Rights Reserved.
+
+#include "ViewportTransformer.h"
+#include "ViewportWorldInteraction.h"
+
+
+void UViewportTransformer::Init( UViewportWorldInteraction* InitViewportWorldInteraction )
+{
+	this->ViewportWorldInteraction = InitViewportWorldInteraction;
+}
+
+
+void UViewportTransformer::Shutdown()
+{
+	this->ViewportWorldInteraction = nullptr;
+}
+
diff --git a/Engine/Source/Editor/ViewportInteraction/ViewportWorldInteraction.cpp b/Engine/Source/Editor/ViewportInteraction/ViewportWorldInteraction.cpp
--- a/Engine/Source/Editor/ViewportInteraction/ViewportWorldInteraction.cpp
+++ b/Engine/Source/Editor/ViewportInteraction/ViewportWorldInteraction.cpp
@@ -21,7 +21,8 @@
 #include "Gizmo/VIPivotTransformGizmo.h"
 #include "ViewportWorldInteractionManager.h"
 #include "IViewportInteractionModule.h"
-
+#include "ViewportTransformer.h"
+#include "ActorTransformer.h"
 #include "SnappingUtils.h"
 #include "ScopedTransaction.h"
 
@@ -49,16 +50,17 @@ namespace VI
 	static FAutoConsoleVariable DragAtLaserImpactInterpolationDuration( TEXT( "VI.DragAtLaserImpactInterpolationDuration" ), 0.1f, TEXT( "How long we should interpolate objects between positions when dragging under the laser's impact point" ) );
 	static FAutoConsoleVariable DragAtLaserImpactInterpolationThreshold( TEXT( "VI.DragAtLaserImpactInterpolationThreshold" ), 5.0f, TEXT( "Minimum distance jumped between frames before we'll force interpolation mode to activated" ) );
 
-	static FAutoConsoleVariable ForceGizmoPivotToCenterOfSelectedActorsBounds( TEXT( "VI.ForceGizmoPivotToCenterOfSelectedActorsBounds" ), 0, TEXT( "When enabled, the gizmo's pivot will always be centered on the selected actors.  Otherwise, we use the pivot of the last selected actor." ) );
+	static FAutoConsoleVariable ForceGizmoPivotToCenterOfObjectsBounds( TEXT( "VI.ForceGizmoPivotToCenterOfObjectsBounds" ), 0, TEXT( "When enabled, the gizmo's pivot will always be centered on the selected objects.  Otherwise, we use the pivot of the last selected object." ) );
 	static FAutoConsoleVariable GizmoHandleHoverScale( TEXT( "VI.GizmoHandleHoverScale" ), 1.5f, TEXT( "How much to scale up transform gizmo handles when hovered over" ) );
 	static FAutoConsoleVariable GizmoHandleHoverAnimationDuration( TEXT( "VI.GizmoHandleHoverAnimationDuration" ), 0.1f, TEXT( "How quickly to animate gizmo handle hover state" ) );
 	static FAutoConsoleVariable ShowTransformGizmo( TEXT( "VI.ShowTransformGizmo" ), 1, TEXT( "Whether the transform gizmo should be shown for selected objects" ) );
-
-	//Grid
+	static FAutoConsoleVariable DragTranslationVelocityStopEpsilon( TEXT( "VI.DragTranslationVelocityStopEpsilon" ), KINDA_SMALL_NUMBER, TEXT( "When dragging inertia falls below this value (cm/frame), we'll stop inertia and finalize the drag" ) );
 	static FAutoConsoleVariable SnapGridSize( TEXT( "VI.SnapGridSize" ), 3.0f, TEXT( "How big the snap grid should be, in multiples of the gizmo bounding box size" ) );
 	static FAutoConsoleVariable SnapGridLineWidth( TEXT( "VI.SnapGridLineWidth" ), 3.0f, TEXT( "Width of the grid lines on the snap grid" ) );
-	static FAutoConsoleVariable MinVelocityForInertia( TEXT( "VREd.MinVelocityForInertia" ), 1.0f, TEXT( "Minimum velocity (in cm/frame in unscaled room space) before inertia will kick in when releasing objects (or the world)" ) );
+	static FAutoConsoleVariable MinVelocityForInertia( TEXT( "VI.MinVelocityForInertia" ), 1.0f, TEXT( "Minimum velocity (in cm/frame in unscaled room space) before inertia will kick in when releasing objects (or the world)" ) );
 	static FAutoConsoleVariable GridHapticFeedbackStrength( TEXT( "VI.GridHapticFeedbackStrength" ), 0.4f, TEXT( "Default strength for haptic feedback when moving across grid points" ) );
+
+	static FAutoConsoleVariable UseTransientActors( TEXT( "VI.UseTransientActors" ), 1, TEXT( "For debugging only, allows you to turn off transient on newly-spawned VR editor actors." ) );
 }
 
 // @todo vreditor: Hacky inline implementation of a double vector.  Move elsewhere and polish or get rid.
@@ -218,17 +220,19 @@ static void SegmentDistToSegmentDouble( DVector A1, DVector B1, DVector A2, DVec
 	OutP2 = A2 + S2 * T2;
 }
 
-UViewportWorldInteraction::UViewportWorldInteraction( const FObjectInitializer& Initializer ):
-	Super( Initializer ),
+
+UViewportWorldInteraction::UViewportWorldInteraction():
+	Super(),
 	bDraggedSinceLastSelection( false ),
 	LastDragGizmoStartTransform( FTransform::Identity ),
 	TrackingTransaction( FTrackingTransaction() ),
 	AppTimeEntered( FTimespan::Zero() ),
-	EditorViewportClient( nullptr ),
+	DefaultOptionalViewportClient( nullptr ),
 	LastFrameNumberInputWasPolled( 0 ),
 	MotionControllerID( 0 ),	// @todo ViewportInteraction: We only support a single controller, and we assume the first controller are the motion controls
 	LastWorldToMetersScale( 100.0f ),
 	bSkipInteractiveWorldMovementThisFrame( false ),
+	bAreTransformablesMoving( false ),
 	bIsInterpolatingTransformablesFromSnapshotTransform( false ),
 	TransformablesInterpolationStartTime( FTimespan::Zero() ),
 	TransformablesInterpolationDuration( 1.0f ),
@@ -237,25 +241,22 @@ UViewportWorldInteraction::UViewportWorldInteraction( const FObjectInitializer&
 	GizmoLocalBounds( FBox( 0 ) ),
 	StartDragAngleOnRotation(),
 	StartDragHandleDirection(),
-	CurrentGizmoType( EGizmoHandleTypes::All ),
-	bIsTransformGizmoVisible( true ),
+	bShouldTransformGizmoBeVisible( true ),
+	TransformGizmoScale( 0.5f ),
 	SnapGridActor( nullptr ),
 	SnapGridMeshComponent( nullptr ),
 	SnapGridMID( nullptr ),
 	DraggedInteractable( nullptr ),
-	bActive( false )
+	DefaultMouseCursorInteractor( nullptr ),
+	DefaultMouseCursorInteractorRefCount( 0 ),
+	bActive( false ),
+	bUseInputProcessor( false )
 {
 }
 
-UViewportWorldInteraction::~UViewportWorldInteraction()
-{
-	Shutdown();
-}
 
-void UViewportWorldInteraction::Init( UWorld* InWorld )
+void UViewportWorldInteraction::Init()
 {
-	World = InWorld;
-
 	Colors.SetNumZeroed( (int32)EColors::TotalCount );
 	{
 		Colors[ (int32)EColors::DefaultColor ] = FLinearColor( 0.7f, 0.7f, 0.7f, 1.0f );;
@@ -267,10 +268,34 @@ void UViewportWorldInteraction::Init( UWorld* InWorld )
 	}
 
 	AppTimeEntered = FTimespan::FromSeconds( FApp::GetCurrentTime() );
+
+	IViewportWorldInteractionManager& WorldInteractionManager = IViewportInteractionModule::Get().GetWorldInteractionManager();
+	WorldInteractionManager.SetCurrentViewportWorldInteraction( this );
+
+	// Start with the default transformer
+	SetTransformer( nullptr );
+
+	//Spawn the transform gizmo at init so we do not hitch when selecting our first object
+	SpawnTransformGizmoIfNeeded();
+	TransformGizmoActor->SetIsTemporarilyHiddenInEditor( true );
+
+
+	/** This will make sure this is not ticking after the editor has been closed. */
+	GEditor->OnEditorClose().AddUObject( this, &UViewportWorldInteraction::Shutdown );
+
+	bActive = true;
 }
 
 void UViewportWorldInteraction::Shutdown()
 {
+	bActive = false;
+
+	DestroyActors();
+	DefaultOptionalViewportClient = nullptr;
+
+	IViewportWorldInteractionManager& WorldInteractionManager = IViewportInteractionModule::Get().GetWorldInteractionManager();
+	WorldInteractionManager.SetCurrentViewportWorldInteraction( nullptr );
+
 	AppTimeEntered = FTimespan::Zero();
 
 	for ( UViewportInteractor* Interactor : Interactors )
@@ -289,43 +314,18 @@ void UViewportWorldInteraction::Shutdown()
 	OnAxisInputEvent.Clear();
 	OnStopDraggingEvent.Clear();
 
-	World = nullptr;
 	TransformGizmoActor = nullptr;
 	SnapGridActor = nullptr;
 	SnapGridMeshComponent = nullptr;
 	SnapGridMID = nullptr;
-	EditorViewportClient = nullptr;
 	DraggedInteractable = nullptr;
+	if( ViewportTransformer != nullptr )
+	{
+		ViewportTransformer->Shutdown();
+		ViewportTransformer = nullptr;
+	}
 
-	bActive = false;
-	USelection::SelectionChangedEvent.RemoveAll( this );
-}
-
-void UViewportWorldInteraction::Enter()
-{
-	IViewportWorldInteractionManager& WorldInteractionManager = IViewportInteractionModule::Get().GetWorldInteractionManager();
-	WorldInteractionManager.SetCurrentViewportWorldInteraction( this );
-
-	// Find out about selection changes
-	USelection::SelectionChangedEvent.AddUObject( this, &UViewportWorldInteraction::OnActorSelectionChanged ); 
-	
-    //Spawn the transform gizmo at init so we do not hitch when selecting our first object
-	SpawnTransformGizmoIfNeeded();
-	TransformGizmoActor->SetIsTemporarilyHiddenInEditor( true );
-
-	/** This will make sure this is not ticking after the editor has been closed. */
-	GEditor->OnEditorClose().AddUObject( this, &UViewportWorldInteraction::Shutdown );
-}
-
-void UViewportWorldInteraction::Exit()
-{
-	DestroyActors();
 	USelection::SelectionChangedEvent.RemoveAll( this );
-	EditorViewportClient = nullptr;
-	
-	IViewportWorldInteractionManager& WorldInteractionManager = IViewportInteractionModule::Get().GetWorldInteractionManager();
-	WorldInteractionManager.SetCurrentViewportWorldInteraction( nullptr );
-
 	GEditor->OnEditorClose().RemoveAll( this );
 }
 
@@ -333,18 +333,13 @@ void UViewportWorldInteraction::OnEditorClosed()
 {
 	if( bActive )
 	{
-		Exit();
 		Shutdown();
 	}
 }
 
-void UViewportWorldInteraction::Tick( const float DeltaTime )
+void UViewportWorldInteraction::Tick( float DeltaSeconds )
 {
-	// Only if this is our viewport. Remember, editor modes get a chance to tick and receive input for each active viewport.
-	if ( EditorViewportClient == nullptr || !bActive )
-	{
-		return;
-	}
+	OnPreWorldInteractionTickEvent.Broadcast( DeltaSeconds );
 
 	// Cancel any haptic effects that have been playing too long
 	//StopOldHapticEffects(); //@todo ViewportInteraction
@@ -353,15 +348,17 @@ void UViewportWorldInteraction::Tick( const float DeltaTime )
 	PollInputIfNeeded();
 
 	// Update hover. Note that hover can also be updated when ticking our sub-systems below.
-	HoverTick( DeltaTime );
+	HoverTick( DeltaSeconds );
 
-	InteractionTick( DeltaTime );
+	InteractionTick( DeltaSeconds );
 
 	// Update all the interactors
 	for ( UViewportInteractor* Interactor : Interactors )
 	{
-		Interactor->Tick( DeltaTime );
+		Interactor->Tick( DeltaSeconds );
 	}
+
+	OnPostWorldInteractionTickEvent.Broadcast( DeltaSeconds );
 }
 
 void UViewportWorldInteraction::AddInteractor( UViewportInteractor* Interactor )
@@ -376,26 +373,51 @@ void UViewportWorldInteraction::RemoveInteractor( UViewportInteractor* Interacto
 	Interactors.Remove( Interactor );
 }
 
-void UViewportWorldInteraction::SetViewport(const TSharedPtr<class FEditorViewportClient>& InEditorViewportClient)
-{
-	EditorViewportClient = InEditorViewportClient.Get();
-}
 
-void UViewportWorldInteraction::Activate(const bool bInActivate)
+void UViewportWorldInteraction::SetTransformer( UViewportTransformer* NewTransformer )
 {
-	if( bActive != bInActivate )
+	// Clear all existing transformables
+	SetTransformables( TArray< TUniquePtr< class FViewportTransformable > >() );
+
+	if( ViewportTransformer != nullptr )
 	{
-		if( bInActivate )
-		{
-			Enter();
-		}
-		else
-		{
-			Exit();
-		}
+		ViewportTransformer->Shutdown();
+		ViewportTransformer = nullptr;
+	}
+
+	this->ViewportTransformer = NewTransformer;
 
-		bActive = bInActivate;
+	if( ViewportTransformer == nullptr )
+	{
+		ViewportTransformer = NewObject<UActorTransformer>();
 	}
+	check( ViewportTransformer != nullptr );
+
+	ViewportTransformer->Init( this );
+}
+
+
+void UViewportWorldInteraction::SetTransformables( TArray< TUniquePtr< class FViewportTransformable > >&& NewTransformables )
+{
+	// We're dealing with a new set of transformables, so clear out anything that was going on with the old ones.
+	bDraggedSinceLastSelection = false;
+	LastDragGizmoStartTransform = FTransform::Identity;
+	bIsInterpolatingTransformablesFromSnapshotTransform = false;
+	TransformablesInterpolationStartTime = FTimespan::Zero();
+	TransformablesInterpolationDuration = 1.0f;
+
+	Transformables = MoveTemp( NewTransformables );
+
+	const bool bNewObjectsSelected = true;
+	const bool bAllHandlesVisible = true;
+	UActorComponent* SingleVisibleHandle = nullptr;
+	TArray<UActorComponent*> HoveringOverHandles;
+	RefreshTransformGizmo( bNewObjectsSelected, bAllHandlesVisible, SingleVisibleHandle, HoveringOverHandles );
+}
+
+void UViewportWorldInteraction::SetDefaultOptionalViewportClient(const TSharedPtr<class FEditorViewportClient>& InEditorViewportClient)
+{
+	DefaultOptionalViewportClient = InEditorViewportClient.Get();
 }
 
 bool UViewportWorldInteraction::IsActive() const
@@ -409,20 +431,23 @@ void UViewportWorldInteraction::PairInteractors( UViewportInteractor* FirstInter
 	SecondInteractor->SetOtherInteractor( FirstInteractor );
 }
 
-bool UViewportWorldInteraction::HandleInputKey( const FKey Key, const EInputEvent Event )
+bool UViewportWorldInteraction::InputKey( FEditorViewportClient* InViewportClient, FViewport* Viewport, FKey Key, EInputEvent Event )
 {
 	bool bWasHandled = false;
+
 	if( bActive )
 	{
-		OnKeyInputEvent.Broadcast( *EditorViewportClient, Key, Event, bWasHandled );
+		check( InViewportClient != nullptr );
+		OnKeyInputEvent.Broadcast( InViewportClient, Key, Event, bWasHandled );
+
 		if( !bWasHandled || !bActive )
 		{
 			for( UViewportInteractor* Interactor : Interactors )
 			{
-				bWasHandled = Interactor->HandleInputKey( Key, Event );
-				
+				bWasHandled = Interactor->HandleInputKey( *InViewportClient, Key, Event );
+
 				// Stop iterating if the input was handled by an Interactor
-				if(bWasHandled || !bActive)
+				if( bWasHandled || !bActive )
 				{
 					break;
 				}
@@ -430,25 +455,47 @@ bool UViewportWorldInteraction::HandleInputKey( const FKey Key, const EInputEven
 		}
 	}
 
+	if( bWasHandled )
+	{
+		// This is a temporary workaround so that we don't steal the handling of the drag tool from the viewport client
+		FInputEventState InputState( InViewportClient->Viewport, Key, Event );
+		const bool bMarqueeSelecting = InputState.IsAltButtonPressed() && InputState.IsCtrlButtonPressed() && Key == EKeys::LeftMouseButton;
+		if( bMarqueeSelecting )
+		{
+			bWasHandled = false;
+		}
+	}
+
 	return bWasHandled;
 }
 
-bool UViewportWorldInteraction::HandleInputAxis( const int32 ControllerId, const FKey Key, const float Delta, const float DeltaTime )
+bool UViewportWorldInteraction::PreprocessedInputKey( const FKey Key, const EInputEvent Event )
+{
+	if( DefaultOptionalViewportClient != nullptr && bUseInputProcessor == true )
+	{
+		return InputKey( DefaultOptionalViewportClient, DefaultOptionalViewportClient->Viewport, Key, Event );
+	}
+
+	return false;
+}
+
+bool UViewportWorldInteraction::InputAxis( FEditorViewportClient* InViewportClient, FViewport* Viewport, int32 ControllerId, FKey Key, float Delta, float DeltaTime )
 {
 	bool bWasHandled = false;
 
-	if(bActive)
+	if( bActive )
 	{
-		OnAxisInputEvent.Broadcast( *EditorViewportClient, ControllerId, Key, Delta, DeltaTime, bWasHandled );
+		check( InViewportClient != nullptr );
+		OnAxisInputEvent.Broadcast( InViewportClient, ControllerId, Key, Delta, DeltaTime, bWasHandled );
 
-		if( !bWasHandled || !bActive)
+		if( !bWasHandled || !bActive )
 		{
 			for( UViewportInteractor* Interactor : Interactors )
 			{
-				bWasHandled = Interactor->HandleInputAxis( Key, Delta, DeltaTime );
+				bWasHandled = Interactor->HandleInputAxis( *InViewportClient, Key, Delta, DeltaTime );
 
 				// Stop iterating if the input was handled by an interactor
-				if(bWasHandled || !bActive)
+				if( bWasHandled || !bActive )
 				{
 					break;
 				}
@@ -459,22 +506,33 @@ bool UViewportWorldInteraction::HandleInputAxis( const int32 ControllerId, const
 	return bWasHandled;
 }
 
+bool UViewportWorldInteraction::PreprocessedInputAxis( const int32 ControllerId, const FKey Key, const float Delta, const float DeltaTime )
+{
+	if( DefaultOptionalViewportClient != nullptr && bUseInputProcessor == true )
+	{
+		return InputAxis( DefaultOptionalViewportClient, DefaultOptionalViewportClient->Viewport, ControllerId, Key, Delta, DeltaTime );
+	}
+
+	return false;
+}
 
 FTransform UViewportWorldInteraction::GetRoomTransform() const
 {
-	check( EditorViewportClient );
-
-	const FTransform RoomTransform(
-		EditorViewportClient->GetViewRotation().Quaternion(),
-		EditorViewportClient->GetViewLocation(),
-		FVector( 1.0f ) );
+	FTransform RoomTransform;
+	if( DefaultOptionalViewportClient != nullptr )
+	{
+		RoomTransform = FTransform(
+			DefaultOptionalViewportClient->GetViewRotation().Quaternion(),
+			DefaultOptionalViewportClient->GetViewLocation(),
+			FVector( 1.0f ) );
+	}
 	return RoomTransform;
 }
 
 FTransform UViewportWorldInteraction::GetRoomSpaceHeadTransform() const
 {
 	FTransform HeadTransform = FTransform::Identity;
-	if ( GEngine->HMDDevice.IsValid() )
+	if( HaveHeadTransform() )
 	{
 		FQuat RoomSpaceHeadOrientation;
 		FVector RoomSpaceHeadLocation;
@@ -495,74 +553,84 @@ FTransform UViewportWorldInteraction::GetHeadTransform() const
 	return GetRoomSpaceHeadTransform() * GetRoomTransform();
 }
 
+
+bool UViewportWorldInteraction::HaveHeadTransform() const
+{
+	return DefaultOptionalViewportClient != nullptr && GEngine->HMDDevice.IsValid() && GEngine->HMDDevice->IsStereoEnabled();
+}
+
+
 void UViewportWorldInteraction::SetRoomTransform( const FTransform& NewRoomTransform )
 {
-	EditorViewportClient->SetViewLocation( NewRoomTransform.GetLocation() );
-	EditorViewportClient->SetViewRotation( NewRoomTransform.GetRotation().Rotator() );
+	if( DefaultOptionalViewportClient != nullptr )
+	{
+		DefaultOptionalViewportClient->SetViewLocation( NewRoomTransform.GetLocation() );
+		DefaultOptionalViewportClient->SetViewRotation( NewRoomTransform.GetRotation().Rotator() );
 
-	// Forcibly dirty the viewport camera location
-	const bool bDollyCamera = false;
-	EditorViewportClient->MoveViewportCamera( FVector::ZeroVector, FRotator::ZeroRotator, bDollyCamera );
+		// Forcibly dirty the viewport camera location
+		const bool bDollyCamera = false;
+		DefaultOptionalViewportClient->MoveViewportCamera( FVector::ZeroVector, FRotator::ZeroRotator, bDollyCamera );
+	}
 }
 
 float UViewportWorldInteraction::GetWorldScaleFactor() const
 {
-	return GetViewportWorld()->GetWorldSettings()->WorldToMeters / 100.0f;
+	return GetWorld()->GetWorldSettings()->WorldToMeters / 100.0f;
 }
 
-UWorld* UViewportWorldInteraction::GetViewportWorld() const
+UWorld* UViewportWorldInteraction::GetWorld() const
 {
 	UWorld* ResultWorld = nullptr;
 	
-	if( World )
+	if( UEditorWorldExtension::GetWorld() != nullptr )
 	{
-		ResultWorld = World;
+		ResultWorld = UEditorWorldExtension::GetWorld();
 	}
-	else if ( EditorViewportClient != nullptr )
+	else if ( DefaultOptionalViewportClient != nullptr )
 	{
-		ResultWorld = EditorViewportClient->GetWorld();
+		ResultWorld = DefaultOptionalViewportClient->GetWorld();
 	}
 
 	return ResultWorld;
 }
 
-FEditorViewportClient* UViewportWorldInteraction::GetViewportClient() const
+FEditorViewportClient* UViewportWorldInteraction::GetDefaultOptionalViewportClient() const
 {
-	return EditorViewportClient;
+	return DefaultOptionalViewportClient;
 }
 
 void UViewportWorldInteraction::Undo()
 {
-	GUnrealEd->Exec( GetViewportWorld(), TEXT( "TRANSACTION UNDO" ) );
+	GUnrealEd->Exec( GetWorld(), TEXT( "TRANSACTION UNDO" ) );
 }
 
 void UViewportWorldInteraction::Redo()
 {
-	GUnrealEd->Exec( GetViewportWorld(), TEXT( "TRANSACTION REDO" ) );
+	GUnrealEd->Exec( GetWorld(), TEXT( "TRANSACTION REDO" ) );
 }
 
 void UViewportWorldInteraction::DeleteSelectedObjects()
 {
-	GUnrealEd->Exec(GetViewportWorld(), TEXT("DELETE"));
+	GUnrealEd->Exec(GetWorld(), TEXT("DELETE"));
 }
 
 void UViewportWorldInteraction::Copy()
 {
 	// @todo vreditor: Needs CanExecute()  (see LevelEditorActions.cpp)
-	GUnrealEd->Exec( GetViewportWorld(), TEXT( "EDIT COPY" ) );
+	GUnrealEd->Exec( GetWorld(), TEXT( "EDIT COPY" ) );
 }
 
 void UViewportWorldInteraction::Paste()
 {
 	// @todo vreditor: Needs CanExecute()  (see LevelEditorActions.cpp)
 	// @todo vreditor: Needs "paste here" style pasting (TO=HERE), but with ray
-	GUnrealEd->Exec( GetViewportWorld(), TEXT( "EDIT PASTE" ) );
+	GUnrealEd->Exec( GetWorld(), TEXT( "EDIT PASTE" ) );
 }
 
 void UViewportWorldInteraction::Duplicate()
 {
 	// @todo vreditor: Needs CanExecute()  (see LevelEditorActions.cpp)
-	GUnrealEd->Exec( GetViewportWorld(), TEXT( "DUPLICATE" ) );
+	GUnrealEd->Exec( GetWorld(), TEXT( "DUPLICATE" ) );
 }
 
 void UViewportWorldInteraction::Deselect()
@@ -582,9 +650,22 @@ void UViewportWorldInteraction::HoverTick( const float DeltaTime )
 		//Update the state of the interactor here
 		FHitResult HitHoverResult;
 		Interactor->UpdateHoverResult( HitHoverResult );
-
 		FVector HoverImpactPoint = HitHoverResult.ImpactPoint;
-		OnHoverUpdateEvent.Broadcast( *EditorViewportClient, Interactor, /* In/Out */ HoverImpactPoint, /* In/Out */ bWasHandled);
+
+		const bool bIsHoveringOverTransformGizmo =
+			HitHoverResult.Actor.IsValid() &&
+			HitHoverResult.Actor == TransformGizmoActor;
+
+		// Prefer transform gizmo hover over everything else
+		if( !bIsHoveringOverTransformGizmo )
+		{
+			OnHoverUpdateEvent.Broadcast( Interactor, /* In/Out */ HoverImpactPoint, /* In/Out */ bWasHandled );
+			if( bWasHandled )
+			{
+				InteractorData.bIsHovering = true;
+				InteractorData.HoverLocation = HoverImpactPoint;
+			}
+		}
 		
 		if ( bWasHandled )
 		{
@@ -664,20 +745,19 @@ void UViewportWorldInteraction::HoverTick( const float DeltaTime )
 void UViewportWorldInteraction::InteractionTick( const float DeltaTime )
 {
 	const FTimespan CurrentTime = FTimespan::FromSeconds( FPlatformTime::Seconds() );
-	const float WorldToMetersScale = GetViewportWorld()->GetWorldSettings()->WorldToMeters;
+	const float WorldToMetersScale = GetWorld()->GetWorldSettings()->WorldToMeters;
 
 	// Update viewport with any objects that are currently hovered over
 	{
-		const bool bUseEditorSelectionHoverFeedback = GEditor != NULL && GetDefault<ULevelEditorViewportSettings>()->bEnableViewportHoverFeedback;
-		if( bUseEditorSelectionHoverFeedback )
+		if( DefaultOptionalViewportClient != nullptr )
 		{
-			if( EditorViewportClient->IsLevelEditorClient() )
+			const bool bUseEditorSelectionHoverFeedback = GEditor != NULL && GetDefault<ULevelEditorViewportSettings>()->bEnableViewportHoverFeedback;
+			if( bUseEditorSelectionHoverFeedback && DefaultOptionalViewportClient->IsLevelEditorClient())
 			{
-				FLevelEditorViewportClient* LevelEditorViewportClient = static_cast<FLevelEditorViewportClient*>( EditorViewportClient );
+				FLevelEditorViewportClient* LevelEditorViewportClient = static_cast<FLevelEditorViewportClient*>( DefaultOptionalViewportClient );
 				LevelEditorViewportClient->UpdateHoveredObjects( HoveredObjects );
 			}
 		}
-
 		// This will be filled in again during the next input update
 		HoveredObjects.Reset();
 	}
@@ -685,22 +765,22 @@ void UViewportWorldInteraction::InteractionTick( const float DeltaTime )
 	const float WorldScaleFactor = GetWorldScaleFactor();
 
 	// Move selected actors
-	TOptional< UViewportInteractor* > DraggingActorsWithInteractor;
-	TOptional< UViewportInteractor* > AssistingDragWithInteractor;
+	UViewportInteractor* DraggingTransformablesWithInteractor = nullptr;
+	UViewportInteractor* AssistingDragWithInteractor = nullptr;
 
 	for ( UViewportInteractor* Interactor : Interactors )
 	{
 		bool bCanSlideRayLength = false;
 		FViewportInteractorData& InteractorData = Interactor->GetInteractorData();
 
-		if ( InteractorData.DraggingMode == EViewportInteractionDraggingMode::ActorsWithGizmo ||
-			InteractorData.DraggingMode == EViewportInteractionDraggingMode::ActorsFreely ||
-			InteractorData.DraggingMode == EViewportInteractionDraggingMode::ActorsAtLaserImpact )
+		if( InteractorData.DraggingMode == EViewportInteractionDraggingMode::TransformablesWithGizmo ||
+			InteractorData.DraggingMode == EViewportInteractionDraggingMode::TransformablesFreely ||
+			InteractorData.DraggingMode == EViewportInteractionDraggingMode::TransformablesAtLaserImpact )
 		{
-			check( !DraggingActorsWithInteractor.IsSet() );	// Only support dragging one thing at a time right now!
-			DraggingActorsWithInteractor = Interactor;
+			check( DraggingTransformablesWithInteractor == nullptr );	// Only support dragging one thing at a time right now!
+			DraggingTransformablesWithInteractor = Interactor;
 
-			if ( InteractorData.DraggingMode != EViewportInteractionDraggingMode::ActorsAtLaserImpact )
+			if ( InteractorData.DraggingMode != EViewportInteractionDraggingMode::TransformablesAtLaserImpact )
 			{
 				bCanSlideRayLength = true;
 			}
@@ -708,7 +788,7 @@ void UViewportWorldInteraction::InteractionTick( const float DeltaTime )
 
 		if ( InteractorData.DraggingMode == EViewportInteractionDraggingMode::AssistingDrag )
 		{
-			check( !AssistingDragWithInteractor.IsSet() );	// Only support assisting one thing at a time right now!
+			check( AssistingDragWithInteractor == nullptr );	// Only support assisting one thing at a time right now!
 			AssistingDragWithInteractor = Interactor;
 
 			bCanSlideRayLength = true;
@@ -719,21 +799,24 @@ void UViewportWorldInteraction::InteractionTick( const float DeltaTime )
 			Interactor->CalculateDragRay();
 		}
 	}
-	
+
+	const bool bIsStillDraggingTransformables = DraggingTransformablesWithInteractor != nullptr;
+
+	bool bAllInteractorsThatWereDraggingTransformablesHaveZeroVelocity = true;	
 	for ( UViewportInteractor* Interactor : Interactors )
 	{
 		FViewportInteractorData& InteractorData = Interactor->GetInteractorData();
 			
-		if( InteractorData.DraggingMode == EViewportInteractionDraggingMode::ActorsWithGizmo ||
-			InteractorData.DraggingMode == EViewportInteractionDraggingMode::ActorsFreely ||
-			InteractorData.DraggingMode == EViewportInteractionDraggingMode::ActorsAtLaserImpact ||
+		if( InteractorData.DraggingMode == EViewportInteractionDraggingMode::TransformablesWithGizmo ||
+			InteractorData.DraggingMode == EViewportInteractionDraggingMode::TransformablesFreely ||
+			InteractorData.DraggingMode == EViewportInteractionDraggingMode::TransformablesAtLaserImpact ||
 			InteractorData.DraggingMode == EViewportInteractionDraggingMode::World )
 		{
 			// Are we dragging with two interactors ?
 			UViewportInteractor* OtherInteractor = nullptr;
-			if ( AssistingDragWithInteractor.IsSet() )
+			if ( AssistingDragWithInteractor != nullptr )
 			{
-				OtherInteractor = AssistingDragWithInteractor.GetValue();
+				OtherInteractor = AssistingDragWithInteractor;
 			}
 
 			UViewportInteractor* OtherInteractorThatWasAssistingDrag = GetOtherInteractorIntertiaContribute( Interactor );
@@ -763,16 +846,16 @@ void UViewportWorldInteraction::InteractionTick( const float DeltaTime )
 			FVector LaserPointerDirection = InteractorData.Transform.GetUnitAxis( EAxis::X );
 			FVector LaserPointerEnd = InteractorData.Transform.GetLocation();
 
-			if( InteractorData.DraggingMode == EViewportInteractionDraggingMode::ActorsWithGizmo ||
-				InteractorData.DraggingMode == EViewportInteractionDraggingMode::ActorsFreely ||
-				InteractorData.DraggingMode == EViewportInteractionDraggingMode::ActorsAtLaserImpact )
+			if( InteractorData.DraggingMode == EViewportInteractionDraggingMode::TransformablesWithGizmo ||
+				InteractorData.DraggingMode == EViewportInteractionDraggingMode::TransformablesFreely ||
+				InteractorData.DraggingMode == EViewportInteractionDraggingMode::TransformablesAtLaserImpact )
 			{
 				// Move objects using the laser pointer (in world space)
 				if( Interactor->GetLaserPointer( /* Out */ LaserPointerStart, /* Out */ LaserPointerEnd ) )
 				{
 					LaserPointerDirection = ( LaserPointerEnd - LaserPointerStart ).GetSafeNormal();
 
-					if( InteractorData.DraggingMode == EViewportInteractionDraggingMode::ActorsAtLaserImpact )
+					if( InteractorData.DraggingMode == EViewportInteractionDraggingMode::TransformablesAtLaserImpact )
 					{
 						// Check to see if the laser pointer is over something we can drop on
 						FVector HitLocation = FVector::ZeroVector;
@@ -901,19 +984,18 @@ void UViewportWorldInteraction::InteractionTick( const float DeltaTime )
 
 			{
 				const bool bIsMouseCursorInteractor = Cast<UMouseCursorInteractor>( Interactor ) != nullptr;
-				// @todo vreditor: We could do a world space distance test (times world scale factor) when forcing VR mode to get similar (but not quite the same) behavior
 				{
 					// Don't bother with inertia if we're not moving very fast.  This filters out tiny accidental movements.
 					const FVector RoomSpaceHandDelta = bIsMouseCursorInteractor ?
-						( InteractorData.Transform.GetLocation() - InteractorData.LastTransform.GetLocation() ) :
+						( DragDelta ) :	// For the mouse cursor the interactor origin won't change unless the camera moves, so just test the distance we dragged instead.
 						( InteractorData.RoomSpaceTransform.GetLocation() - InteractorData.LastRoomSpaceTransform.GetLocation() );
 					if( RoomSpaceHandDelta.Size() < VI::MinVelocityForInertia->GetFloat() * WorldScaleFactor )
 					{
 						InteractorData.DragTranslationVelocity = FVector::ZeroVector;
 					}
-					if( AssistingDragWithInteractor.IsSet() )
+					if( AssistingDragWithInteractor != nullptr )
 					{
-						FViewportInteractorData& AssistingOtherInteractorData = AssistingDragWithInteractor.GetValue()->GetInteractorData();
+						FViewportInteractorData& AssistingOtherInteractorData = AssistingDragWithInteractor->GetInteractorData();
 						const FVector OtherHandRoomSpaceHandDelta = ( AssistingOtherInteractorData.RoomSpaceTransform.GetLocation() - AssistingOtherInteractorData.LastRoomSpaceTransform.GetLocation() );
 						if( OtherHandRoomSpaceHandDelta.Size() < VI::MinVelocityForInertia->GetFloat() * WorldScaleFactor )
 						{
@@ -923,7 +1005,8 @@ void UViewportWorldInteraction::InteractionTick( const float DeltaTime )
 				}
 			}
 
-			const FVector OldViewLocation = EditorViewportClient->GetViewLocation();
+			const FVector OldViewLocation = DefaultOptionalViewportClient != nullptr ? DefaultOptionalViewportClient->GetViewLocation() : FVector::ZeroVector;
+
 			// Dragging transform gizmo handle
 			const bool bWithTwoHands = ( OtherInteractor != nullptr || OtherInteractorThatWasAssistingDrag != nullptr );
 			const bool bIsLaserPointerValid = true;
@@ -934,7 +1017,7 @@ void UViewportWorldInteraction::InteractionTick( const float DeltaTime )
 				InteractorData.DraggingMode, 
 				InteractorData.TransformGizmoInteractionType, 
 				bWithTwoHands,
-				*EditorViewportClient, 
+				*DefaultOptionalViewportClient, 
 				InteractorData.OptionalHandlePlacement, 
 				DragDelta, 
 				OtherHandDragDelta,
@@ -944,6 +1027,7 @@ void UViewportWorldInteraction::InteractionTick( const float DeltaTime )
 				OtherHandDragDeltaFromStart,
 				LaserPointerStart, 
 				LaserPointerDirection, 
+				Interactor->GetLaserPointerMaxLength(),
 				bIsLaserPointerValid, 
 				InteractorData.GizmoStartTransform, 
 				InteractorData.GizmoStartLocalBounds,
@@ -956,7 +1040,7 @@ void UViewportWorldInteraction::InteractionTick( const float DeltaTime )
 
 			// Make sure the hover point is right on the position that we're dragging the object to.  This is important
 			// when constraining dragged objects to a single axis or a plane
-			if( InteractorData.DraggingMode == EViewportInteractionDraggingMode::ActorsWithGizmo )
+			if( InteractorData.DraggingMode == EViewportInteractionDraggingMode::TransformablesWithGizmo )
 			{
 				InteractorData.HoverLocation = FMath::ClosestPointOnSegment( UnsnappedDraggedTo, LaserPointerStart, LaserPointerEnd );
 			}
@@ -973,7 +1057,7 @@ void UViewportWorldInteraction::InteractionTick( const float DeltaTime )
 			// If we were dragging the world, then play some haptic feedback
 			if( InteractorData.DraggingMode == EViewportInteractionDraggingMode::World )
 			{
-				const FVector NewViewLocation = EditorViewportClient->GetViewLocation();
+				const FVector NewViewLocation = DefaultOptionalViewportClient != nullptr ? DefaultOptionalViewportClient->GetViewLocation() : FVector::ZeroVector;
 
 				// @todo vreditor: Consider doing this for inertial moves too (we need to remember the last hand that invoked the move.)
 
@@ -1025,10 +1109,17 @@ void UViewportWorldInteraction::InteractionTick( const float DeltaTime )
 		// If we're not actively dragging, apply inertia to any selected elements that we've dragged around recently
 		else 
 		{
-			if( !InteractorData.DragTranslationVelocity.IsNearlyZero( KINDA_SMALL_NUMBER ) &&
+			if( !InteractorData.DragTranslationVelocity.IsNearlyZero( VI::DragTranslationVelocityStopEpsilon->GetFloat() ) &&
 				!InteractorData.bWasAssistingDrag && 	// If we were only assisting, let the other hand take care of doing the update
 				!InteractorData.bIsDrivingVelocityOfSimulatedTransformables )	// If simulation mode is on, let the physics engine take care of inertia
 			{
+				if( InteractorData.LastDraggingMode == EViewportInteractionDraggingMode::TransformablesWithGizmo ||
+					InteractorData.LastDraggingMode == EViewportInteractionDraggingMode::TransformablesFreely ||
+					InteractorData.LastDraggingMode == EViewportInteractionDraggingMode::TransformablesAtLaserImpact )
+				{
+					bAllInteractorsThatWereDraggingTransformablesHaveZeroVelocity = false;
+				}
+
 				const FVector DragDelta = InteractorData.DragTranslationVelocity;
 				const FVector DraggedTo = InteractorData.LastDragToLocation + DragDelta;
 				const FVector DragDeltaFromStart = DraggedTo - InteractorData.LaserPointerImpactAtDragStart;
@@ -1056,7 +1147,7 @@ void UViewportWorldInteraction::InteractionTick( const float DeltaTime )
 					InteractorData.LastDraggingMode, 
 					InteractorData.TransformGizmoInteractionType, 
 					bWithTwoHands, 
-					*EditorViewportClient, 
+					*DefaultOptionalViewportClient, 
 					InteractorData.OptionalHandlePlacement, 
 					DragDelta, 
 					OtherHandDragDelta,
@@ -1066,6 +1157,7 @@ void UViewportWorldInteraction::InteractionTick( const float DeltaTime )
 					OtherHandDragDeltaFromStart,
 					FVector::ZeroVector,	// No valid laser pointer during inertia
 					FVector::ZeroVector,	// No valid laser pointer during inertia
+					0.0f,		// No valid laser pointer during inertia
 					bIsLaserPointerValid, 
 					InteractorData.GizmoStartTransform, 
 					InteractorData.GizmoStartLocalBounds, 
@@ -1098,8 +1190,8 @@ void UViewportWorldInteraction::InteractionTick( const float DeltaTime )
 	// measurement text labels), so it just makes sense to keep it up to date.
 	{
 		const bool bNewObjectsSelected = false;
-		const bool bAllHandlesVisible = ( !DraggingActorsWithInteractor.IsSet() );
-		UActorComponent* SingleVisibleHandle = ( DraggingActorsWithInteractor.IsSet() ) ? DraggingActorsWithInteractor.GetValue()->GetInteractorData().DraggingTransformGizmoComponent.Get() : nullptr;
+		const bool bAllHandlesVisible = ( DraggingTransformablesWithInteractor == nullptr );
+		UActorComponent* SingleVisibleHandle = ( DraggingTransformablesWithInteractor != nullptr ) ? DraggingTransformablesWithInteractor->GetInteractorData().DraggingTransformGizmoComponent.Get() : nullptr;
 
 		static TArray< UActorComponent* > HoveringOverHandles;
 		HoveringOverHandles.Reset();
@@ -1118,105 +1210,28 @@ void UViewportWorldInteraction::InteractionTick( const float DeltaTime )
 
 
 	// Update transformables
-	//UpdateTransformables( DeltaTime );
 	const bool bSmoothSnappingEnabled = IsSmoothSnappingEnabled();
-	if ( bSmoothSnappingEnabled || bIsInterpolatingTransformablesFromSnapshotTransform )
+	if ( bAreTransformablesMoving && ( bSmoothSnappingEnabled || bIsInterpolatingTransformablesFromSnapshotTransform ) )
 	{
 		const float SmoothSnapSpeed = VI::SmoothSnapSpeed->GetFloat();
-		const bool bUseElasticSnapping = bSmoothSnappingEnabled && VI::ElasticSnap->GetInt() > 0 && TrackingTransaction.IsActive();	// Only while we're still dragging stuff!
+		const bool bUseElasticSnapping = bSmoothSnappingEnabled && VI::ElasticSnap->GetInt() > 0 && bIsStillDraggingTransformables;	// Only while we're still dragging stuff!
 		const float ElasticSnapStrength = VI::ElasticSnapStrength->GetFloat();
 
-		// NOTE: We never sweep while smooth-snapping as the object will never end up where we want it to
-		// due to friction with the ground and other objects.
-		const bool bSweep = false;
+		bool bIsStillSmoothSnapping = false;
 
-		for ( TUniquePtr<FViewportTransformable>& TransformablePtr : Transformables )
+		float InterpProgress = 1.0f;
+		const bool bWasInterpolationNeeded = bIsInterpolatingTransformablesFromSnapshotTransform;
+		bool bDidInterpolationFinishThisUpdate = false;
+		if( bIsInterpolatingTransformablesFromSnapshotTransform )
 		{
-			FViewportTransformable& Transformable = *TransformablePtr;
-
-			FTransform ActualTargetTransform = Transformable.TargetTransform;
-
-			// If 'elastic snapping' is turned on, we'll have the object 'reach' toward its unsnapped position, so
-			// that the user always has visual feedback when they are dragging something around, even if they
-			// haven't dragged far enough to overcome the snap threshold.
-			if ( bUseElasticSnapping )
-			{
-				ActualTargetTransform.BlendWith( Transformable.UnsnappedTargetTransform, ElasticSnapStrength );
-			}
-
-			FTransform InterpolatedTransform = ActualTargetTransform;
-			if ( !ActualTargetTransform.Equals( Transformable.LastTransform, 0.0f ) )
-			{
-				// If we're really close, just snap it.
-				if ( ActualTargetTransform.Equals( Transformable.LastTransform, KINDA_SMALL_NUMBER ) )
-				{
-					InterpolatedTransform = Transformable.LastTransform = Transformable.UnsnappedTargetTransform = ActualTargetTransform;
-				}
-				else
-				{
-					InterpolatedTransform.Blend(
-						Transformable.LastTransform,
-						ActualTargetTransform,
-						FMath::Min( 1.0f, SmoothSnapSpeed * FMath::Min( DeltaTime, 1.0f / 30.0f ) ) );
-				}
-				Transformable.LastTransform = InterpolatedTransform;
-			}
-
-			if ( bIsInterpolatingTransformablesFromSnapshotTransform )
-			{
-				const float InterpProgress = FMath::Clamp( ( float ) ( CurrentTime - TransformablesInterpolationStartTime ).GetTotalSeconds() / TransformablesInterpolationDuration, 0.0f, 1.0f );
-				InterpolatedTransform.BlendWith( Transformable.InterpolationSnapshotTransform, 1.0f - InterpProgress );
-				if ( InterpProgress >= 1.0f - KINDA_SMALL_NUMBER )
-				{
-					// Finished interpolating
-					bIsInterpolatingTransformablesFromSnapshotTransform = false;
-				}
-			}
-
-			// Got an actor?
-			AActor* Actor = Cast<AActor>( Transformable.Object.Get() );
-			if ( Actor != nullptr )
+			InterpProgress = FMath::Clamp( (float)( CurrentTime - TransformablesInterpolationStartTime ).GetTotalSeconds() / TransformablesInterpolationDuration, 0.0f, 1.0f );
+			if( InterpProgress >= 1.0f - KINDA_SMALL_NUMBER )
 			{
-				const FTransform& ExistingTransform = Actor->GetActorTransform();
-				if ( !ExistingTransform.Equals( InterpolatedTransform, 0.0f ) )
-				{
-					const bool bOnlyTranslationChanged =
-						ExistingTransform.GetRotation() == InterpolatedTransform.GetRotation() &&
-						ExistingTransform.GetScale3D() == InterpolatedTransform.GetScale3D();
-
-					Actor->SetActorTransform( InterpolatedTransform, bSweep );
-					//GWarn->Logf( TEXT( "SMOOTH: Actor %s to %s" ), *Actor->GetName(), *Transformable.TargetTransform.ToString() );
-
-					Actor->InvalidateLightingCacheDetailed( bOnlyTranslationChanged );
-
-					const bool bFinished = false;
-					Actor->PostEditMove( bFinished );
-				}
-			}
-			else
-			{
-				// Some other object that we don't know how to drag
+				// Finished interpolating
+				bIsInterpolatingTransformablesFromSnapshotTransform = false;
+				bDidInterpolationFinishThisUpdate = true;
 			}
 		}
-	}
-
-	LastWorldToMetersScale = WorldToMetersScale;
-} 
-
-void UViewportWorldInteraction::UpdateTransformables( const float DeltaTime )
-{
-	const FTimespan CurrentTime = FTimespan::FromSeconds( FPlatformTime::Seconds() );
-
-	const bool bSmoothSnappingEnabled = IsSmoothSnappingEnabled();
-	if ( bSmoothSnappingEnabled || bIsInterpolatingTransformablesFromSnapshotTransform )
-	{
-		const float SmoothSnapSpeed = VI::SmoothSnapSpeed->GetFloat();
-		const bool bUseElasticSnapping = bSmoothSnappingEnabled && VI::ElasticSnap->GetInt() > 0 && TrackingTransaction.IsActive();	// Only while we're still dragging stuff!
-		const float ElasticSnapStrength = VI::ElasticSnapStrength->GetFloat();
-
-		// NOTE: We never sweep while smooth-snapping as the object will never end up where we want it to
-		// due to friction with the ground and other objects.
-		const bool bSweep = false;
 
 		for ( TUniquePtr<FViewportTransformable>& TransformablePtr : Transformables )
 		{
@@ -1242,6 +1257,8 @@ void UViewportWorldInteraction::UpdateTransformables( const float DeltaTime )
 				}
 				else
 				{
+					bIsStillSmoothSnapping = true;
+
 					InterpolatedTransform.Blend(
 						Transformable.LastTransform,
 						ActualTargetTransform,
@@ -1252,47 +1269,40 @@ void UViewportWorldInteraction::UpdateTransformables( const float DeltaTime )
 
 			if ( bIsInterpolatingTransformablesFromSnapshotTransform )
 			{
-				const float InterpProgress = FMath::Clamp( ( float ) ( CurrentTime - TransformablesInterpolationStartTime ).GetTotalSeconds() / TransformablesInterpolationDuration, 0.0f, 1.0f );
 				InterpolatedTransform.BlendWith( Transformable.InterpolationSnapshotTransform, 1.0f - InterpProgress );
-				if ( InterpProgress >= 1.0f - KINDA_SMALL_NUMBER )
-				{
-					// Finished interpolating
-					bIsInterpolatingTransformablesFromSnapshotTransform = false;
-				}
 			}
 
-			// Got an actor?
-			AActor* Actor = Cast<AActor>( Transformable.Object.Get() );
-			if ( Actor != nullptr )
-			{
-				const FTransform& ExistingTransform = Actor->GetActorTransform();
-				if ( !ExistingTransform.Equals( InterpolatedTransform, 0.0f ) )
-				{
-					const bool bOnlyTranslationChanged =
-						ExistingTransform.GetRotation() == InterpolatedTransform.GetRotation() &&
-						ExistingTransform.GetScale3D() == InterpolatedTransform.GetScale3D();
-
-					Actor->SetActorTransform( InterpolatedTransform, bSweep );
-					Actor->InvalidateLightingCacheDetailed( bOnlyTranslationChanged );
+			// Update the transform!
+			// NOTE: We never sweep while smooth-snapping as the object will never end up where we want it to
+			// due to friction with the ground and other objects.
+			const bool bSweep = false;
+			Transformable.ApplyTransform( InterpolatedTransform, bSweep );
+		}
 
-					const bool bFinished = false;
-					Actor->PostEditMove( bFinished );
-				}
-			}
-			else
-			{
-				// Some other object that we don't know how to drag
-			}
+		// Did we finish interpolating in this update?  If so, we need to finalize things and notify everyone
+		if( bAreTransformablesMoving &&
+			!bIsStillDraggingTransformables &&
+			bAllInteractorsThatWereDraggingTransformablesHaveZeroVelocity &&
+			( !bWasInterpolationNeeded || bDidInterpolationFinishThisUpdate ) && 
+			!bIsStillSmoothSnapping )
+		{
+			FinishedMovingTransformables();
+		}
+	}
+	else
+	{
+		// Neither smooth snapping or interpolation is enabled right now, but the transformables could have some velocity
+		// applied.  We'll check to see whether they've come to a rest, and if so finalize their positions.
+		if( bAreTransformablesMoving && 
+			!bIsStillDraggingTransformables &&
+			bAllInteractorsThatWereDraggingTransformablesHaveZeroVelocity )
+		{
+			FinishedMovingTransformables();
 		}
 	}
-}
-
-void UViewportWorldInteraction::RefreshOnSelectionChanged( const bool bNewObjectsSelected, bool bAllHandlesVisible, class UActorComponent* SingleVisibleHandle, const TArray< UActorComponent* >& HoveringOverHandles )
-{
-	//@todo: ViewportInteraction: Call listeners for this event
 
-	RefreshTransformGizmo( bNewObjectsSelected, bAllHandlesVisible, SingleVisibleHandle, HoveringOverHandles );
-}
+	LastWorldToMetersScale = WorldToMetersScale;
+} 
 
 bool UViewportWorldInteraction::IsInteractableComponent( const UActorComponent* Component ) const
 {
@@ -1337,6 +1347,7 @@ void UViewportWorldInteraction::UpdateDragging(
 	const FVector& OtherHandDragDeltaFromStart,
 	const FVector& LaserPointerStart,
 	const FVector& LaserPointerDirection,
+	const float LaserPointerMaxLength,
 	const bool bIsLaserPointerValid,
 	const FTransform& GizmoStartTransform,
 	const FBox& GizmoStartLocalBounds,
@@ -1363,14 +1374,14 @@ void UViewportWorldInteraction::UpdateDragging(
 	// NOTE: While placing objects, we always snap in world space, since the initial position isn't really at all useful
 	const bool bLocalSpaceSnapping =
 		GetTransformGizmoCoordinateSpace() == COORD_Local &&
-		DraggingMode != EViewportInteractionDraggingMode::ActorsAtLaserImpact;
-	const FVector SnapGridBase = ( DraggingMode == EViewportInteractionDraggingMode::ActorsAtLaserImpact||  bLocalSpaceSnapping ) ? FVector::ZeroVector : GizmoStartTransform.GetLocation(); //@todo ViewportInteraction
+		DraggingMode != EViewportInteractionDraggingMode::TransformablesAtLaserImpact;
+	const FVector SnapGridBase = ( DraggingMode == EViewportInteractionDraggingMode::TransformablesAtLaserImpact ||  bLocalSpaceSnapping ) ? FVector::ZeroVector : GizmoStartTransform.GetLocation(); //@todo ViewportInteraction
 
 
 	// Okay, time to move stuff!  We'll do this part differently depending on whether we're dragging actual actors around
 	// or we're moving the camera (aka. dragging the world)
-	if ( DraggingMode == EViewportInteractionDraggingMode::ActorsWithGizmo ||
-		DraggingMode == EViewportInteractionDraggingMode::ActorsAtLaserImpact )
+	if ( DraggingMode == EViewportInteractionDraggingMode::TransformablesWithGizmo ||
+		DraggingMode == EViewportInteractionDraggingMode::TransformablesAtLaserImpact )
 	{
 		FVector OutClosestPointOnLaser;
 		FVector ConstrainedDragDeltaFromStart = ComputeConstrainedDragDeltaFromStart(
@@ -1382,6 +1393,7 @@ void UViewportWorldInteraction::UpdateDragging(
 			LaserPointerDirection,
 			bIsLaserPointerValid,
 			GizmoStartTransform,
+			LaserPointerMaxLength,
 			/* In/Out */ GizmoSpaceFirstDragUpdateOffsetAlongAxis,
 			/* In/Out */ DragDeltaFromStartOffset,
 			/* Out */ OutClosestPointOnLaser);
@@ -1716,8 +1728,8 @@ void UViewportWorldInteraction::UpdateDragging(
 			}
 		}
 	}
-	else if ( DraggingMode == EViewportInteractionDraggingMode::ActorsFreely ||
-		DraggingMode == EViewportInteractionDraggingMode::World )
+	else if ( DraggingMode == EViewportInteractionDraggingMode::TransformablesFreely ||
+			  DraggingMode == EViewportInteractionDraggingMode::World )
 	{
 		FVector TranslationOffset = FVector::ZeroVector;
 		FQuat RotationOffset = FQuat::Identity;
@@ -1766,10 +1778,10 @@ void UViewportWorldInteraction::UpdateDragging(
 			if ( false )
 			{
 				const FTransform LinesRelativeTo = DraggingMode == EViewportInteractionDraggingMode::World ? GetRoomTransform() : FTransform::Identity;
-				DrawDebugLine( GetViewportWorld(), LinesRelativeTo.TransformPosition( LineStart ), LinesRelativeTo.TransformPosition( LineEnd ), FColor::Green, false, 0.0f );
-				DrawDebugLine( GetViewportWorld(), LinesRelativeTo.TransformPosition( LastLineStart ), LinesRelativeTo.TransformPosition( LastLineEnd ), FColor::Red, false, 0.0f );
+				DrawDebugLine( GetWorld(), LinesRelativeTo.TransformPosition( LineStart ), LinesRelativeTo.TransformPosition( LineEnd ), FColor::Green, false, 0.0f );
+				DrawDebugLine( GetWorld(), LinesRelativeTo.TransformPosition( LastLineStart ), LinesRelativeTo.TransformPosition( LastLineEnd ), FColor::Red, false, 0.0f );
 
-				DrawDebugSphere( GetViewportWorld(), LinesRelativeTo.TransformPosition( PivotLocation ), 2.5f * GetWorldScaleFactor(), 32, FColor::White, false, 0.0f );
+				DrawDebugSphere( GetWorld(), LinesRelativeTo.TransformPosition( PivotLocation ), 2.5f * GetWorldScaleFactor(), 32, FColor::White, false, 0.0f );
 			}
 
 			const float LastLineLength = ( LastLineEnd - LastLineStart ).Size();
@@ -1798,7 +1810,7 @@ void UViewportWorldInteraction::UpdateDragging(
 		}
 
 
-		if ( DraggingMode == EViewportInteractionDraggingMode::ActorsFreely )
+		if ( DraggingMode == EViewportInteractionDraggingMode::TransformablesFreely )
 		{
 			for ( TUniquePtr<FViewportTransformable>& TransformablePtr : Transformables )
 			{
@@ -1897,7 +1909,7 @@ void UViewportWorldInteraction::UpdateDragging(
 			const float WorldScaleOffset = ScaleOffset.GetAbsMax();
 			if ( WorldScaleOffset != 0.0f )
 			{
-				const float OldWorldToMetersScale = GetViewportWorld()->GetWorldSettings()->WorldToMeters;
+				const float OldWorldToMetersScale = GetWorld()->GetWorldSettings()->WorldToMeters;
 				const float NewWorldToMetersScale = OldWorldToMetersScale / WorldScaleOffset;
 
 				// NOTE: Instead of clamping, we simply skip changing the W2M this frame if it's out of bounds.  Clamping makes our math more complicated.
@@ -1960,27 +1972,20 @@ void UViewportWorldInteraction::UpdateDragging(
 			{
 				FViewportTransformable& Transformable = *TransformablePtr;
 
-				AActor* Actor = Cast<AActor>( Transformable.Object.Get() );
-				if ( Actor != nullptr )
+				// @todo VREditor physics: When freely transforming, should set angular velocity too (will pivot be the same though??)
+				if( Transformable.IsPhysicallySimulated() )
 				{
-					UPrimitiveComponent* RootPrimitiveComponent = Cast<UPrimitiveComponent>( Actor->GetRootComponent() );
-					if ( RootPrimitiveComponent != nullptr )
+					FVector MoveDelta = Transformable.UnsnappedTargetTransform.GetLocation() - Transformable.LastTransform.GetLocation();
+					if ( bShouldApplyVelocitiesFromDrag )
 					{
-						// @todo VREditor physics: When freely transforming, should set angular velocity too (will pivot be the same though??)
-						if ( RootPrimitiveComponent->IsSimulatingPhysics() )
-						{
-							FVector MoveDelta = Transformable.UnsnappedTargetTransform.GetLocation() - Transformable.LastTransform.GetLocation();
-							if ( bShouldApplyVelocitiesFromDrag )
-							{
-								bIsDrivingVelocityOfSimulatedTransformables = true;
-							}
-							else
-							{
-								MoveDelta = FVector::ZeroVector;
-							}
-							RootPrimitiveComponent->SetAllPhysicsLinearVelocity( MoveDelta * DeltaTime * VI::InertiaVelocityBoost->GetFloat() );
-						}
+						bIsDrivingVelocityOfSimulatedTransformables = true;
+					}
+					else
+					{
+						MoveDelta = FVector::ZeroVector;
 					}
+
+					Transformable.SetLinearVelocity( MoveDelta * DeltaTime * VI::InertiaVelocityBoost->GetFloat() );
 				}
 			}
 		}
@@ -1995,27 +2000,8 @@ void UViewportWorldInteraction::UpdateDragging(
 
 				Transformable.LastTransform = Transformable.TargetTransform;
 
-				// Got an actor?
-				AActor* Actor = Cast<AActor>( Transformable.Object.Get() );
-				if ( Actor != nullptr )
-				{
-					const FTransform& ExistingTransform = Actor->GetTransform();
-
-					const bool bOnlyTranslationChanged =
-						ExistingTransform.GetRotation() == Transformable.TargetTransform.GetRotation() &&
-						ExistingTransform.GetScale3D() == Transformable.TargetTransform.GetScale3D();
-
-					Actor->SetActorTransform( Transformable.TargetTransform, bSweep );
-
-					Actor->InvalidateLightingCacheDetailed( bOnlyTranslationChanged );
-
-					const bool bFinished = false;
-					Actor->PostEditMove( bFinished );
-				}
-				else
-				{
-					// Some other object that we don't know how to drag
-				}
+				// Update the transform!
+				Transformable.ApplyTransform( Transformable.TargetTransform, bSweep );
 			}
 		}
 	}
@@ -2034,6 +2020,7 @@ FVector UViewportWorldInteraction::ComputeConstrainedDragDeltaFromStart(
 	const FVector& LaserPointerDirection, 
 	const bool bIsLaserPointerValid, 
 	const FTransform& GizmoStartTransform, 
+	const float LaserPointerMaxLength,
 	FVector& GizmoSpaceFirstDragUpdateOffsetAlongAxis, 
 	FVector& GizmoSpaceDragDeltaFromStartOffset,
 	FVector& OutClosestPointOnLaser) const
@@ -2056,8 +2043,6 @@ FVector UViewportWorldInteraction::ComputeConstrainedDragDeltaFromStart(
 			const FVector GizmoSpaceLaserPointerStart = GizmoStartTransform.InverseTransformPosition( LaserPointerStart );
 			const FVector GizmoSpaceLaserPointerDirection = GizmoStartTransform.InverseTransformVectorNoScale( LaserPointerDirection );
 
-			const float LaserPointerMaxLength = 1000; //Owner.GetLaserPointerMaxLength(); //@todo viewportinteraction: Needs proper way of getting distance
-
 			FVector GizmoSpaceConstraintAxis =
 				FacingAxisIndex == 0 ? FVector::ForwardVector : ( FacingAxisIndex == 1 ? FVector::RightVector : FVector::UpVector );
 			if ( HandlePlacement.Axes[ FacingAxisIndex ] == ETransformGizmoHandleDirection::Negative )
@@ -2096,7 +2081,7 @@ FVector UViewportWorldInteraction::ComputeConstrainedDragDeltaFromStart(
 
 				DVector GizmoSpaceClosestPointOnLaserDouble, GizmoSpaceClosestPointOnAxisDouble;
 				SegmentDistToSegmentDouble(
-					GizmoSpaceLaserPointerStart, DVector( GizmoSpaceLaserPointerStart ) + DVector( GizmoSpaceLaserPointerDirection ) * LaserPointerMaxLength,	// @todo vreditor: Should store laser pointer length rather than hard code
+					GizmoSpaceLaserPointerStart, DVector( GizmoSpaceLaserPointerStart ) + DVector( GizmoSpaceLaserPointerDirection ) * LaserPointerMaxLength,
 					DVector( GizmoSpaceConstraintAxis ) * -AxisSegmentLength, DVector( GizmoSpaceConstraintAxis ) * AxisSegmentLength,
 					/* Out */ GizmoSpaceClosestPointOnLaserDouble,
 					/* Out */ GizmoSpaceClosestPointOnAxisDouble );
@@ -2174,43 +2159,42 @@ FVector UViewportWorldInteraction::ComputeConstrainedDragDeltaFromStart(
 	return ConstrainedWorldSpaceDeltaFromStart;
 }
 
-void UViewportWorldInteraction::StartDraggingActors( UViewportInteractor* Interactor, const FViewportActionKeyInput& Action, UActorComponent* ClickedComponent, const FVector& HitLocation, const bool bIsPlacingActors )
+bool UViewportWorldInteraction::StartDragging( UViewportInteractor* Interactor, UActorComponent* ClickedTransformGizmoComponent, const FVector& HitLocation, const bool bIsPlacingNewObjects, const bool bStartTransaction )
 {
-	if ( IsInteractableComponent( ClickedComponent ) )
+	bool bStartedDragging = false;
 	{
 		FVector LaserPointerStart, LaserPointerEnd;
 		const bool bHaveLaserPointer = Interactor->GetLaserPointer( /* Out */ LaserPointerStart, /* Out */ LaserPointerEnd );
 		if ( bHaveLaserPointer )
 		{
-			AActor* Actor = ClickedComponent->GetOwner(); 
 			FViewportInteractorData& InteractorData = Interactor->GetInteractorData();
 
-			// Capture undo state
-			if ( bIsPlacingActors )
-			{
-				// When placing actors, a transaction should already be in progress
-				check( TrackingTransaction.IsActive() );
-			}
-			else
+			if( bStartTransaction )
 			{
+				// Capture undo state
 				TrackingTransaction.TransCount++;
-				TrackingTransaction.Begin( LOCTEXT( "MovingActors", "Moving Selected Actors" ) );
+				TrackingTransaction.Begin( LOCTEXT( "MovingObjects", "Move Object" ) );
 
 				// Suspend actor/component modification during each delta step to avoid recording unnecessary overhead into the transaction buffer
 				GEditor->DisableDeltaModification( true );
 			}
+			else
+			{
+				// A transaction should already be in progress
+				check( GEditor->IsTransactionActive() );
+			}
 
-			// Give the interacter a chance to do something when starting dragging
-			Interactor->OnStartDragging( ClickedComponent, HitLocation, bIsPlacingActors );
+			// Give the interactor a chance to do something when starting dragging
+			Interactor->OnStartDragging( HitLocation, bIsPlacingNewObjects );
 
-			const bool bUsingGizmo =
-				Action.ActionType == ViewportWorldActionTypes::SelectAndMove_LightlyPressed &&		// Only use the gizmo when lightly pressed
-				( Actor == TransformGizmoActor ) &&
-				ClickedComponent != nullptr;
+			const bool bUsingGizmo = ClickedTransformGizmoComponent != nullptr;
+			check( !bUsingGizmo || ( ClickedTransformGizmoComponent->GetOwner() == TransformGizmoActor ) );
 
 			// Start dragging the objects right away!
-			InteractorData.DraggingMode = InteractorData.LastDraggingMode = bIsPlacingActors ? EViewportInteractionDraggingMode::ActorsAtLaserImpact : 
-				( bUsingGizmo ? EViewportInteractionDraggingMode::ActorsWithGizmo : EViewportInteractionDraggingMode::ActorsFreely );
+			InteractorData.DraggingMode = InteractorData.LastDraggingMode = bIsPlacingNewObjects ? EViewportInteractionDraggingMode::TransformablesAtLaserImpact :
+				( bUsingGizmo ? EViewportInteractionDraggingMode::TransformablesWithGizmo : EViewportInteractionDraggingMode::TransformablesFreely );
+
+			bAreTransformablesMoving = true;
 
 			// Starting a new drag, so make sure the other hand doesn't think it's assisting us
 			if( Interactor->GetOtherInteractor() != nullptr )
@@ -2233,7 +2217,7 @@ void UViewportWorldInteraction::StartDraggingActors( UViewportInteractor* Intera
 			// the gizmo to transform it.
 			if ( bUsingGizmo )
 			{
-				InteractorData.DraggingTransformGizmoComponent = Cast<USceneComponent>( ClickedComponent );
+				InteractorData.DraggingTransformGizmoComponent = Cast<USceneComponent>( ClickedTransformGizmoComponent );
 			}
 			else
 			{
@@ -2258,16 +2242,16 @@ void UViewportWorldInteraction::StartDraggingActors( UViewportInteractor* Intera
 			bDraggedSinceLastSelection = true;
 			LastDragGizmoStartTransform = InteractorData.GizmoStartTransform;
 
-			SetupTransformablesForSelectedActors();
+			// Make sure all of our transformables are reset to their initial transform before we start dragging them
+			ResetTransformables();
 
 			// Calculate the offset between the average location and the hit location after setting the transformables for the new selected objects
 			InteractorData.StartHitLocationToTransformableCenter = CalculateAverageLocationOfTransformables() - HitLocation;
 
 			// If we're placing actors, start interpolating to their actual location.  This helps smooth everything out when
 			// using the laser impact point as the target transform
-			if ( bIsPlacingActors )
+			if( bIsPlacingNewObjects )
 			{
-				GEditor->BroadcastBeginObjectMovement(*Actor);
 				bIsInterpolatingTransformablesFromSnapshotTransform = true;
 				const FTimespan CurrentTime = FTimespan::FromSeconds( FPlatformTime::Seconds() );
 				TransformablesInterpolationStartTime = CurrentTime;
@@ -2276,16 +2260,24 @@ void UViewportWorldInteraction::StartDraggingActors( UViewportInteractor* Intera
 
 			// Play a haptic effect when objects are picked up
 			Interactor->PlayHapticEffect( Interactor->GetDragHapticFeedbackStrength() );
+
+			OnStartDraggingEvent.Broadcast( Interactor );
+
+			bStartedDragging = true;
 		}
 	}
+
+	return bStartedDragging;
 }
+
+
 void UViewportWorldInteraction::StopDragging( UViewportInteractor* Interactor )
 {
-	OnStopDraggingEvent.Broadcast( Interactor );
-
 	FViewportInteractorData& InteractorData = Interactor->GetInteractorData();
 	if ( InteractorData.DraggingMode != EViewportInteractionDraggingMode::Nothing )
 	{
+		OnStopDraggingEvent.Broadcast( Interactor );
+
 		// If the other hand started dragging after we started, allow that hand to "take over" the drag, so the user
 		// doesn't have to click again to continue their action.  Inertial effects of the hand that stopped dragging
 		// will still be in effect.
@@ -2322,9 +2314,9 @@ void UViewportWorldInteraction::StopDragging( UViewportInteractor* Interactor )
 					}
 				}
 			}
-			else if ( InteractorData.DraggingMode == EViewportInteractionDraggingMode::ActorsWithGizmo ||
-				InteractorData.DraggingMode == EViewportInteractionDraggingMode::ActorsFreely ||
-				InteractorData.DraggingMode == EViewportInteractionDraggingMode::ActorsAtLaserImpact )
+			else if ( InteractorData.DraggingMode == EViewportInteractionDraggingMode::TransformablesWithGizmo ||
+					  InteractorData.DraggingMode == EViewportInteractionDraggingMode::TransformablesFreely ||
+					  InteractorData.DraggingMode == EViewportInteractionDraggingMode::TransformablesAtLaserImpact )
 			{
 				// vreditor todo: Move this to where we detect a move from inertia finishes
 				FHitResult HitResult = Interactor->GetHitResultFromLaserPointer();
@@ -2333,20 +2325,23 @@ void UViewportWorldInteraction::StopDragging( UViewportInteractor* Interactor )
 					GEditor->BroadcastEndObjectMovement(*HitResult.Actor);
 				}
 				// No hand is dragging the objects anymore.
-				if ( InteractorData.DraggingMode == EViewportInteractionDraggingMode::ActorsWithGizmo &&
-					InteractorData.TransformGizmoInteractionType == ETransformGizmoInteractionType::RotateOnAngle )
+				if ( InteractorData.DraggingMode == EViewportInteractionDraggingMode::TransformablesWithGizmo &&
+					 InteractorData.TransformGizmoInteractionType == ETransformGizmoInteractionType::RotateOnAngle )
 				{
 					StartDragAngleOnRotation.Reset();
 					StartDragHandleDirection.Reset();
 				}
 
-				// Finalize undo
+				// If we're not dragging anything around, check to see if transformables have come to a rest.
+				const bool bTransformablesStillMoving =
+					bAreTransformablesMoving &&
+					( IsSmoothSnappingEnabled() ||
+					  !InteractorData.DragTranslationVelocity.IsNearlyZero( VI::DragTranslationVelocityStopEpsilon->GetFloat() ) ||
+					  bIsInterpolatingTransformablesFromSnapshotTransform );
+				if( !bTransformablesStillMoving )
 				{
-					// @todo vreditor undo: This doesn't actually encapsulate any inertial movement that happens after the drag is released! 
-					// We need to figure out whether that matters or not.  Also, look into PostEditMove( bFinished=true ) and how that relates to this.
-					--TrackingTransaction.TransCount;
-					TrackingTransaction.End();
-					GEditor->DisableDeltaModification( false );
+					// Finalize the drag
+					FinishedMovingTransformables();
 				}	
 			}
 		}
@@ -2357,31 +2352,30 @@ void UViewportWorldInteraction::StopDragging( UViewportInteractor* Interactor )
 	// NOTE: Even though we're done dragging, we keep our list of transformables so that inertial drag can still push them around!
 }
 
-void UViewportWorldInteraction::SetupTransformablesForSelectedActors()
+
+void UViewportWorldInteraction::FinishedMovingTransformables()
 {
-	//@todo ViewportInteraction: Make sure this is not only for actors
+	bAreTransformablesMoving = false;
 	bIsInterpolatingTransformablesFromSnapshotTransform = false;
 	TransformablesInterpolationStartTime = FTimespan::Zero();
 	TransformablesInterpolationDuration = 1.0f;
 
-	Transformables.Reset();
-
-	USelection* ActorSelectionSet = GEditor->GetSelectedActors();
+	OnFinishedMovingTransformablesEvent.Broadcast();
 
-	static TArray<UObject*> SelectedActorObjects;
-	SelectedActorObjects.Reset();
-	ActorSelectionSet->GetSelectedObjects( AActor::StaticClass(), /* Out */ SelectedActorObjects );
-
-	for ( UObject* SelectedActorObject : SelectedActorObjects )
+	// Finalize undo
 	{
-		AActor* SelectedActor = CastChecked<AActor>( SelectedActorObject );
+		// @todo vreditor undo: This doesn't actually encapsulate any inertial movement that happens after the drag is released! 
+		// We need to figure out whether that matters or not.  Also, look into PostEditMove( bFinished=true ) and how that relates to this.
+		--TrackingTransaction.TransCount;
+		TrackingTransaction.End();
+		GEditor->DisableDeltaModification( false );
+	}
+}
 
-		Transformables.Add( TUniquePtr< FViewportTransformable>( new FViewportTransformable() ) );
-		FViewportTransformable& Transformable = *Transformables[ Transformables.Num() - 1 ];
 
-		Transformable.Object = SelectedActor;
-		Transformable.StartTransform = Transformable.LastTransform = Transformable.TargetTransform = Transformable.UnsnappedTargetTransform = Transformable.InterpolationSnapshotTransform = SelectedActor->GetTransform();
-	}
+void UViewportWorldInteraction::SetUseInputProcessor( bool bInUseInputProcessor )
+{
+	bUseInputProcessor = bInUseInputProcessor;
 }
 
 bool UViewportWorldInteraction::FindPlacementPointUnderLaser( UViewportInteractor* Interactor, FVector& OutHitLocation )
@@ -2399,12 +2393,7 @@ bool UViewportWorldInteraction::FindPlacementPointUnderLaser( UViewportInteracto
 		{
 			FViewportTransformable& Transformable = *TransformablePtr;
 
-			// Got an actor?
-			AActor* Actor = Cast<AActor>( Transformable.Object.Get() );
-			if( Actor != nullptr )
-			{
-				IgnoredActors.Add( Actor );
-			}
+			Transformable.UpdateIgnoredActorList( IgnoredActors );
 		} 
 	}
 
@@ -2455,23 +2444,6 @@ FTrackingTransaction& UViewportWorldInteraction::GetTrackingTransaction()
 	return TrackingTransaction;
 }
 
-bool UViewportWorldInteraction::IsTransformingActor( AActor* Actor ) const
-{
-	bool bFoundActor = false;
-	for ( const TUniquePtr<FViewportTransformable>& TransformablePtr : Transformables )
-	{
-		const FViewportTransformable& Transformable = *TransformablePtr;
-
-		if ( Transformable.Object.Get() == Actor )
-		{
-			bFoundActor = true;
-			break;
-		}
-	}
-
-	return bFoundActor;
-}
-
 bool UViewportWorldInteraction::IsSmoothSnappingEnabled() const
 {
 	const float SmoothSnapSpeed = VI::SmoothSnapSpeed->GetFloat();
@@ -2497,30 +2469,19 @@ void UViewportWorldInteraction::PollInputIfNeeded()
 	}
 }
 
-void UViewportWorldInteraction::OnActorSelectionChanged( UObject* ChangedObject )
+
+void UViewportWorldInteraction::ResetTransformables()
 {
-	if ( EditorViewportClient != nullptr )
+	for( TUniquePtr<FViewportTransformable>& Transformable : Transformables )
 	{
-		const bool bNewObjectsSelected = true;
-		const bool bAllHandlesVisible = true;
-		UActorComponent* SingleVisibleHandle = nullptr;
-		TArray<UActorComponent*> HoveringOverHandles;
-		RefreshOnSelectionChanged( bNewObjectsSelected, bAllHandlesVisible, SingleVisibleHandle, HoveringOverHandles );
-
-		// Clear our transformables.  They'll be doing things relative to the gizmo, which will now be in a different place.
-		// @todo vreditor: We could solve this by moving GizmoLocalBounds into the a Transformables base structure, and making sure
-		// that we never talk to TransformGizmoActor (or selected actors) when updating in tick (currently safe!)
-		Transformables.Reset();
-
-		bDraggedSinceLastSelection = false;
-		LastDragGizmoStartTransform = FTransform::Identity;
+		Transformable->StartTransform = Transformable->LastTransform = Transformable->UnsnappedTargetTransform = Transformable->InterpolationSnapshotTransform = Transformable->TargetTransform = Transformable->GetTransform();
 	}
 }
 
 
 void UViewportWorldInteraction::RefreshTransformGizmo( const bool bNewObjectsSelected, bool bAllHandlesVisible, class UActorComponent* SingleVisibleHandle, const TArray< UActorComponent* >& HoveringOverHandles )
 {
-	if ( GEditor->GetSelectedActorCount() > 0 && VI::ShowTransformGizmo->GetInt() != 0 && bIsTransformGizmoVisible )
+	if ( IsTransformGizmoVisible() )
 	{
 		SpawnTransformGizmoIfNeeded();
 
@@ -2530,16 +2491,10 @@ void UViewportWorldInteraction::RefreshTransformGizmo( const bool bNewObjectsSel
 		const ECoordSystem CurrentCoordSystem = GetTransformGizmoCoordinateSpace();
 		const bool bIsWorldSpaceGizmo = ( CurrentCoordSystem == COORD_World );
 
-		USelection* ActorSelectionSet = GEditor->GetSelectedActors();
-
-		static TArray<UObject*> SelectedActorObjects;
-		SelectedActorObjects.Reset();
-		ActorSelectionSet->GetSelectedObjects( AActor::StaticClass(), /* Out */ SelectedActorObjects );
+		FViewportTransformable& LastTransformable = *Transformables.Last();
 
-		AActor* LastSelectedActor = CastChecked<AActor>( SelectedActorObjects.Last() );
-
-		// Use the location and orientation of the last selected actor, to be consistent with the regular editor's gizmo
-		FTransform GizmoToWorld = LastSelectedActor->GetTransform();
+		// Use the location and orientation of the last selected object, to be consistent with the regular editor's gizmo
+		FTransform GizmoToWorld = LastTransformable.GetTransform();
 		GizmoToWorld.RemoveScaling();	// We don't need the pivot itself to be scaled
 
 		if ( bIsWorldSpaceGizmo )
@@ -2547,33 +2502,32 @@ void UViewportWorldInteraction::RefreshTransformGizmo( const bool bNewObjectsSel
 			GizmoToWorld.SetRotation( FQuat::Identity );
 		}
 
-		// Create a gizmo-local bounds around all of the selected actors
-		FBox GizmoSpaceSelectedActorsBounds;
-		GizmoSpaceSelectedActorsBounds.Init();
+		// Create a gizmo-local bounds around all of the selected objects
+		FBox GizmoSpaceSelectedObjectsBounds;
+		GizmoSpaceSelectedObjectsBounds.Init();
 		{
 			const FTransform WorldToGizmo = GizmoToWorld.Inverse();
-			for ( UObject* SelectedActorObject : SelectedActorObjects )
+			for ( TUniquePtr<FViewportTransformable>& TransformablePtr : Transformables )
 			{
-				AActor* SelectedActor = CastChecked<AActor>( SelectedActorObject );
+				FViewportTransformable& Transformable = *TransformablePtr;
 
-				// Figure out the local space bounds of the entire actor
-				const bool bIncludeNonCollidingComponents = false;	// @todo vreditor: Disabled this because it causes lights to have huge bounds
-				const FBox ActorSpaceBounds = SelectedActor->CalculateComponentsBoundingBoxInLocalSpace( bIncludeNonCollidingComponents );
+				// Figure out the local space bounds of the entire object
+				const FBox LocalSpaceBounds = Transformable.GetLocalSpaceBoundingBox();
 
 				// Transform the bounds into the gizmo's space
-				const FTransform ActorToWorld = SelectedActor->GetTransform();
-				const FTransform ActorToGizmo = ActorToWorld * WorldToGizmo;
-				const FBox GizmoSpaceBounds = ActorSpaceBounds.TransformBy( ActorToGizmo );
+				const FTransform TransformableToWorld = Transformable.GetTransform();
+				const FTransform TransformableToGizmo = TransformableToWorld * WorldToGizmo;
+				const FBox GizmoSpaceBounds = LocalSpaceBounds.TransformBy( TransformableToGizmo );
 
 				// Get that bounding box into the local space of the gizmo
-				GizmoSpaceSelectedActorsBounds += GizmoSpaceBounds;
+				GizmoSpaceSelectedObjectsBounds += GizmoSpaceBounds;
 			}
 		}
 
-		if ( VI::ForceGizmoPivotToCenterOfSelectedActorsBounds->GetInt() > 0 )
+		if ( VI::ForceGizmoPivotToCenterOfObjectsBounds->GetInt() > 0 )
 		{
-			GizmoToWorld.SetLocation( GizmoToWorld.TransformPosition( GizmoSpaceSelectedActorsBounds.GetCenter() ) );
-			GizmoSpaceSelectedActorsBounds = GizmoSpaceSelectedActorsBounds.ShiftBy( -GizmoSpaceSelectedActorsBounds.GetCenter() );
+			GizmoToWorld.SetLocation( GizmoToWorld.TransformPosition( GizmoSpaceSelectedObjectsBounds.GetCenter() ) );
+			GizmoSpaceSelectedObjectsBounds = GizmoSpaceSelectedObjectsBounds.ShiftBy( -GizmoSpaceSelectedObjectsBounds.GetCenter() );
 		}
 
 		if ( bNewObjectsSelected )
@@ -2584,9 +2538,22 @@ void UViewportWorldInteraction::RefreshTransformGizmo( const bool bNewObjectsSel
 		const EGizmoHandleTypes GizmoType = GetCurrentGizmoType();
 		const ECoordSystem GizmoCoordinateSpace = GetTransformGizmoCoordinateSpace();
 
-		GizmoLocalBounds = GizmoSpaceSelectedActorsBounds;
+		GizmoLocalBounds = GizmoSpaceSelectedObjectsBounds;
+
+		// If the user has an HMD on, use the head location for gizmo distance-based sizing, otherwise use the active
+		// level editing viewport's camera location.
+		// @todo gizmo: Ideally the gizmo would be sized separately for every viewport it's visible within
+		FVector ViewerLocation = GizmoToWorld.GetLocation();
+		if( HaveHeadTransform() )
+		{
+			ViewerLocation = GetHeadTransform().GetLocation();
+		}
+		else if( GCurrentLevelEditingViewportClient != nullptr )
+		{
+			ViewerLocation = GCurrentLevelEditingViewportClient->GetViewLocation();
+		}
 
-		TransformGizmoActor->UpdateGizmo( GizmoType, GizmoCoordinateSpace, GizmoToWorld, GizmoSpaceSelectedActorsBounds, GetHeadTransform().GetLocation(), bAllHandlesVisible, SingleVisibleHandle,
+		TransformGizmoActor->UpdateGizmo( GizmoType, GizmoCoordinateSpace, GizmoToWorld, GizmoSpaceSelectedObjectsBounds, ViewerLocation, TransformGizmoScale, bAllHandlesVisible, SingleVisibleHandle,
 			HoveringOverHandles, VI::GizmoHandleHoverScale->GetFloat(), VI::GizmoHandleHoverAnimationDuration->GetFloat() );
 
 		// Only draw if snapping is turned on
@@ -2619,18 +2586,18 @@ void UViewportWorldInteraction::RefreshTransformGizmo( const bool bNewObjectsSel
 			for ( UViewportInteractor* Interactor : Interactors )
 			{
 				const FViewportInteractorData& InteractorData = Interactor->GetInteractorData();
-				if ( InteractorData.DraggingMode == EViewportInteractionDraggingMode::ActorsFreely ||
-					InteractorData.DraggingMode == EViewportInteractionDraggingMode::ActorsWithGizmo ||
-					InteractorData.DraggingMode == EViewportInteractionDraggingMode::ActorsAtLaserImpact )
+				if ( InteractorData.DraggingMode == EViewportInteractionDraggingMode::TransformablesFreely ||
+					InteractorData.DraggingMode == EViewportInteractionDraggingMode::TransformablesWithGizmo ||
+					InteractorData.DraggingMode == EViewportInteractionDraggingMode::TransformablesAtLaserImpact )
 				{
 					DraggingMode = InteractorData.DraggingMode;
 					GizmoStartLocationWhileDragging = InteractorData.GizmoStartTransform.GetLocation();
 					break;
 				}
 				else if ( bDraggedSinceLastSelection &&
-					( InteractorData.LastDraggingMode == EViewportInteractionDraggingMode::ActorsFreely ||
-					InteractorData.LastDraggingMode == EViewportInteractionDraggingMode::ActorsWithGizmo ||
-					InteractorData.LastDraggingMode == EViewportInteractionDraggingMode::ActorsAtLaserImpact ) )
+					( InteractorData.LastDraggingMode == EViewportInteractionDraggingMode::TransformablesFreely ||
+					InteractorData.LastDraggingMode == EViewportInteractionDraggingMode::TransformablesWithGizmo ||
+					InteractorData.LastDraggingMode == EViewportInteractionDraggingMode::TransformablesAtLaserImpact ) )
 				{
 					DraggingMode = InteractorData.LastDraggingMode;
 					GizmoStartLocationWhileDragging = LastDragGizmoStartTransform.GetLocation();
@@ -2730,23 +2697,45 @@ void UViewportWorldInteraction::SpawnTransformGizmoIfNeeded()
 
 		// Create the correct gizmo
 		const bool bWithSceneComponent = false;	// We already have our own scene component
-		TransformGizmoActor = CastChecked<ABaseTransformGizmo>( SpawnTransientSceneActor( TransformGizmoClass, TEXT( "PivotTransformGizmo" ), bWithSceneComponent ) );
+		TransformGizmoActor = CastChecked<ABaseTransformGizmo>( SpawnTransientSceneActor( GetWorld(), TransformGizmoClass, TEXT( "PivotTransformGizmo" ), bWithSceneComponent ) );
 
 		check( TransformGizmoActor != nullptr );
 		TransformGizmoActor->SetOwnerWorldInteraction( this );
+
+		if ( !IsTransformGizmoVisible() )
+		{
+			TransformGizmoActor->SetIsTemporarilyHiddenInEditor( true );
+		}
 	}
 }
 
 void UViewportWorldInteraction::SetTransformGizmoVisible( const bool bShouldBeVisible )
 {
-	bIsTransformGizmoVisible = bShouldBeVisible;
+	bShouldTransformGizmoBeVisible = bShouldBeVisible;
 
 	// NOTE: The actual visibility change will be applied the next tick when RefreshTransformGizmo() is called
 }
 
+bool UViewportWorldInteraction::ShouldTransformGizmoBeVisible() const
+{
+	return bShouldTransformGizmoBeVisible;
+}
+
 bool UViewportWorldInteraction::IsTransformGizmoVisible() const
 {
-	return bIsTransformGizmoVisible;
+	return ( Transformables.Num() > 0 && VI::ShowTransformGizmo->GetInt() != 0 && bShouldTransformGizmoBeVisible );
+}
+
+void UViewportWorldInteraction::SetTransformGizmoScale( const float NewScale )
+{
+	TransformGizmoScale = NewScale;
+
+	// NOTE: The actual scale change will be applied the next tick when RefreshTransformGizmo() is called
+}
+
+float UViewportWorldInteraction::GetTransformGizmoScale() const
+{
+	return TransformGizmoScale;
 }
 
 void UViewportWorldInteraction::ApplyVelocityDamping( FVector& Velocity, const bool bVelocitySensitive )
@@ -2789,7 +2778,7 @@ void UViewportWorldInteraction::SetTransformGizmoCoordinateSpace( const ECoordSy
 
 ECoordSystem UViewportWorldInteraction::GetTransformGizmoCoordinateSpace() const
 {
-	const bool bGetRawValue = true;
+	const bool bGetRawValue = false;
 	const ECoordSystem CurrentCoordSystem = GLevelEditorModeTools().GetCoordSystem( bGetRawValue );
 	return CurrentCoordSystem;
 }
@@ -2828,7 +2817,7 @@ UViewportInteractor* UViewportWorldInteraction::GetOtherInteractorIntertiaContri
 			// to contribute inertia!
 			if ( OtherHandInteractorData.DraggingMode == EViewportInteractionDraggingMode::Nothing && OtherHandInteractorData.bWasAssistingDrag )
 			{
-				if ( !OtherHandInteractorData.DragTranslationVelocity.IsNearlyZero( KINDA_SMALL_NUMBER ) )
+				if ( !OtherHandInteractorData.DragTranslationVelocity.IsNearlyZero( VI::DragTranslationVelocityStopEpsilon->GetFloat() ) )
 				{
 					OtherInteractorThatWasAssistingDrag = OtherInteractor;
 				}
@@ -2839,15 +2828,19 @@ UViewportInteractor* UViewportWorldInteraction::GetOtherInteractorIntertiaContri
 	return OtherInteractorThatWasAssistingDrag;
 }
 
-AActor* UViewportWorldInteraction::SpawnTransientSceneActor( TSubclassOf<AActor> ActorClass, const FString& ActorName, const bool bWithSceneComponent ) const
+AActor* UViewportWorldInteraction::SpawnTransientSceneActor( UWorld* World, TSubclassOf<AActor> ActorClass, const FString& ActorName, const bool bWithSceneComponent )
 {
+	check( World != nullptr );
 	const bool bWasWorldPackageDirty = World->GetOutermost()->IsDirty();
 
-	// @todo vreditor: Needs respawn if world changes (map load, etc.)  Will that always restart the editor mode anyway?
 	FActorSpawnParameters ActorSpawnParameters;
-	ActorSpawnParameters.Name = MakeUniqueObjectName( World, ActorClass, *ActorName );	// @todo vreditor: Without this, SpawnActor() can return us an existing PendingKill actor of the same name!  WTF?
+	ActorSpawnParameters.Name = MakeUniqueObjectName( World, ActorClass, *ActorName );
 	ActorSpawnParameters.SpawnCollisionHandlingOverride = ESpawnActorCollisionHandlingMethod::AlwaysSpawn;
-	ActorSpawnParameters.ObjectFlags = EObjectFlags::RF_Transient;
+
+	if( VI::UseTransientActors->GetInt() != 0 )
+	{
+		ActorSpawnParameters.ObjectFlags = EObjectFlags::RF_Transient;
+	}
 
 	check( ActorClass != nullptr );
 	AActor* NewActor = World->SpawnActor< AActor >( ActorClass, ActorSpawnParameters );
@@ -2876,13 +2869,13 @@ void UViewportWorldInteraction::DestroyActors()
 {
 	if(TransformGizmoActor != nullptr)
 	{
-		DestroyTransientActor( TransformGizmoActor );
+		DestroyTransientActor( GetWorld(), TransformGizmoActor );
 		TransformGizmoActor = nullptr;
 	}
 
 	if(SnapGridActor != nullptr)
 	{
-		DestroyTransientActor( SnapGridActor );
+		DestroyTransientActor( GetWorld(), SnapGridActor );
 		SnapGridActor = nullptr;
 		SnapGridMeshComponent = nullptr;
 	}
@@ -2894,8 +2887,9 @@ void UViewportWorldInteraction::DestroyActors()
 	}
 }
 
-void UViewportWorldInteraction::DestroyTransientActor( AActor* Actor ) const
+void UViewportWorldInteraction::DestroyTransientActor( UWorld* World, AActor* Actor )
 {
+	check( World != nullptr );
 	if ( Actor != nullptr )
 	{
 		const bool bWasWorldPackageDirty = World->GetOutermost()->IsDirty();
@@ -2915,12 +2909,47 @@ void UViewportWorldInteraction::DestroyTransientActor( AActor* Actor ) const
 
 EGizmoHandleTypes UViewportWorldInteraction::GetCurrentGizmoType() const
 {
-	return CurrentGizmoType;
+	switch( GLevelEditorModeTools().GetWidgetMode() )
+	{
+		case FWidget::WM_TranslateRotateZ:
+			return EGizmoHandleTypes::All;
+
+		case FWidget::WM_Translate:
+			return EGizmoHandleTypes::Translate;
+
+		case FWidget::WM_Rotate:
+			return EGizmoHandleTypes::Rotate;
+
+		case FWidget::WM_Scale:
+			return EGizmoHandleTypes::Scale;
+	}
+
+	return EGizmoHandleTypes::Translate;
 }
 
 void UViewportWorldInteraction::SetGizmoHandleType( const EGizmoHandleTypes InGizmoHandleType )
 {
-	CurrentGizmoType = InGizmoHandleType;
+	switch( InGizmoHandleType )
+	{
+		case EGizmoHandleTypes::All:
+			// @todo gizmo: This handle type requires an editor preference to be enabled.  Think about how we actually want to expose the 'universal' gizmo type
+			GLevelEditorModeTools().SetWidgetMode( FWidget::WM_TranslateRotateZ );
+			break;
+
+		case EGizmoHandleTypes::Translate:
+			GLevelEditorModeTools().SetWidgetMode( FWidget::WM_Translate );
+			break;
+
+		case EGizmoHandleTypes::Rotate:
+			GLevelEditorModeTools().SetWidgetMode( FWidget::WM_Rotate );
+			break;
+
+		case EGizmoHandleTypes::Scale:
+			GLevelEditorModeTools().SetWidgetMode( FWidget::WM_Scale );
+			break;
+
+		check(0);
+	}
 }
 
 void UViewportWorldInteraction::SetTransformGizmoClass( const TSubclassOf<ABaseTransformGizmo>& NewTransformGizmoClass )
@@ -2962,7 +2991,7 @@ void UViewportWorldInteraction::SpawnGridMeshActor()
 	if( SnapGridActor == nullptr )
 	{
 		const bool bWithSceneComponent = false;
-		SnapGridActor = SpawnTransientSceneActor<AActor>( TEXT( "SnapGrid" ), bWithSceneComponent );
+		SnapGridActor = SpawnTransientSceneActor<AActor>( GetWorld(), TEXT( "SnapGrid" ), bWithSceneComponent );
 
 		SnapGridMeshComponent = NewObject<UStaticMeshComponent>( SnapGridActor );
 		SnapGridActor->AddOwnedComponent( SnapGridMeshComponent );
@@ -3006,4 +3035,33 @@ FLinearColor UViewportWorldInteraction::GetColor( const EColors Color ) const
 	return Colors[ (int32)Color ];
 }
 
-#undef LOCTEXT_NAMESPACE
+void UViewportWorldInteraction::AddMouseCursorInteractor()
+{
+	if( ++DefaultMouseCursorInteractorRefCount == 1 )
+	{
+		// Add a mouse cursor
+		DefaultMouseCursorInteractor = NewObject<UMouseCursorInteractor>();
+		DefaultMouseCursorInteractor->Init();
+		this->AddInteractor( DefaultMouseCursorInteractor );
+	}
+}
+
+void UViewportWorldInteraction::ReleaseMouseCursorInteractor()
+{
+	if( --DefaultMouseCursorInteractorRefCount == 0 )
+	{
+		DefaultMouseCursorInteractorRefCount = 0;
+
+		// Remove mouse cursor
+		DefaultMouseCursorInteractor->Shutdown();
+		this->RemoveInteractor( DefaultMouseCursorInteractor );
+		DefaultMouseCursorInteractor = nullptr;
+	}
+	else
+	{
+		check( DefaultMouseCursorInteractorRefCount >= 0 );
+	}
+}
+
+
+#undef LOCTEXT_NAMESPACE
\ No newline at end of file
diff --git a/Engine/Source/Editor/ViewportInteraction/ViewportWorldInteractionInterface.cpp b/Engine/Source/Editor/ViewportInteraction/ViewportWorldInteractionInterface.cpp
deleted file mode 100644
diff --git a/Engine/Source/Editor/ViewportInteraction/ViewportWorldInteractionManager.cpp b/Engine/Source/Editor/ViewportInteraction/ViewportWorldInteractionManager.cpp
--- a/Engine/Source/Editor/ViewportInteraction/ViewportWorldInteractionManager.cpp
+++ b/Engine/Source/Editor/ViewportInteraction/ViewportWorldInteractionManager.cpp
@@ -21,33 +21,22 @@ void FViewportWorldInteractionManager::AddReferencedObjects( FReferenceCollector
 	Collector.AddReferencedObject( CurrentWorldInteraction );
 }
 
-void FViewportWorldInteractionManager::Tick( float DeltaTime )
-{
-	if( CurrentWorldInteraction && CurrentWorldInteraction->IsActive() )
-	{
-		// First give registered functions to tick before the world interaction is ticked
-		OnPreWorldInteractionTickEvent.Broadcast( DeltaTime );
-		CurrentWorldInteraction->Tick( DeltaTime );
-		OnPostWorldInteractionTickEvent.Broadcast( DeltaTime );
- 	}
-}
-
-bool FViewportWorldInteractionManager::HandleInputKey( const FKey Key, const EInputEvent Event )
+bool FViewportWorldInteractionManager::PreprocessedInputKey( const FKey Key, const EInputEvent Event )
 {
 	bool bResult = false;
-	if(CurrentWorldInteraction && CurrentWorldInteraction->IsActive())
+	if( CurrentWorldInteraction && CurrentWorldInteraction->IsActive() )
 	{
-		bResult = CurrentWorldInteraction->HandleInputKey( Key, Event );
+		bResult = CurrentWorldInteraction->PreprocessedInputKey( Key, Event );
 	}
 	return bResult;
 }
 
-bool FViewportWorldInteractionManager::HandleInputAxis( const int32 ControllerId, const FKey Key, const float Delta, const float DeltaTime )
+bool FViewportWorldInteractionManager::PreprocessedInputAxis( const int32 ControllerId, const FKey Key, const float Delta, const float DeltaTime )
 {
 	bool bResult = false;
-	if(CurrentWorldInteraction && CurrentWorldInteraction->IsActive())
+	if(CurrentWorldInteraction && CurrentWorldInteraction->IsActive() )
 	{
-		bResult = CurrentWorldInteraction->HandleInputAxis( ControllerId, Key, Delta, DeltaTime );
+		bResult = CurrentWorldInteraction->PreprocessedInputAxis( ControllerId, Key, Delta, DeltaTime );
 	}
 	return bResult;
 }
diff --git a/Engine/Source/Editor/ViewportInteraction/ViewportWorldInteractionManager.h b/Engine/Source/Editor/ViewportInteraction/ViewportWorldInteractionManager.h
--- a/Engine/Source/Editor/ViewportInteraction/ViewportWorldInteractionManager.h
+++ b/Engine/Source/Editor/ViewportInteraction/ViewportWorldInteractionManager.h
@@ -19,9 +19,6 @@ public:
 	virtual void AddReferencedObjects( FReferenceCollector& Collector ) override;
 
 	// IViewportWorldInteractionManager overrides
-	virtual void Tick( float DeltaTime ) override;
-	virtual FOnWorldInteractionTick& OnPreWorldInteractionTick() override { return OnPreWorldInteractionTickEvent; };
-	virtual FOnWorldInteractionTick& OnPostWorldInteractionTick() override { return OnPostWorldInteractionTickEvent; };
 	virtual void SetCurrentViewportWorldInteraction( class UViewportWorldInteraction* WorldInteraction ) override;
 
 	//
@@ -29,10 +26,10 @@ public:
 	//
 
 	/** Passes the input key to the current world interaction */
-	bool HandleInputKey( const FKey Key, const EInputEvent Event );
+	bool PreprocessedInputKey( const FKey Key, const EInputEvent Event );
 
 	/** Passes the input axis to the current world interaction */
-	bool HandleInputAxis( const int32 ControllerId, const FKey Key, const float Delta, const float DeltaTime );
+	bool PreprocessedInputAxis( const int32 ControllerId, const FKey Key, const float Delta, const float DeltaTime );
 
 
 private:
@@ -40,16 +37,6 @@ private:
 	/** The current world interaction that is ticking */
 	class UViewportWorldInteraction* CurrentWorldInteraction;
 
-	//
-	// Events
-	//
-
-	/** Event to tick before the worldinteraction ticks */
-	FOnWorldInteractionTick OnPreWorldInteractionTickEvent;
-
-	/** Event to tick after the worldinteraction has ticked */
-	FOnWorldInteractionTick OnPostWorldInteractionTickEvent;
-
 	//
 	// Input
 	//
diff --git a/Engine/Source/Runtime/Engine/Private/ActorEditor.cpp b/Engine/Source/Runtime/Engine/Private/ActorEditor.cpp
--- a/Engine/Source/Runtime/Engine/Private/ActorEditor.cpp
+++ b/Engine/Source/Runtime/Engine/Private/ActorEditor.cpp
@@ -592,8 +592,11 @@ bool AActor::IsHiddenEd() const
 
 void AActor::SetIsTemporarilyHiddenInEditor( bool bIsHidden )
 {
-	bHiddenEdTemporary = bIsHidden;
-	MarkComponentsRenderStateDirty();
+	if( bHiddenEdTemporary != bIsHidden )
+	{
+		bHiddenEdTemporary = bIsHidden;
+		MarkComponentsRenderStateDirty();
+	}
 }
 
 bool AActor::IsEditable() const
diff --git a/Engine/Source/Runtime/Slate/Private/Framework/Application/SlateApplication.cpp b/Engine/Source/Runtime/Slate/Private/Framework/Application/SlateApplication.cpp
--- a/Engine/Source/Runtime/Slate/Private/Framework/Application/SlateApplication.cpp
+++ b/Engine/Source/Runtime/Slate/Private/Framework/Application/SlateApplication.cpp
@@ -4856,6 +4856,12 @@ bool FSlateApplication::ProcessMouseButtonDownEvent( const TSharedPtr< FGenericW
 	}
 	PressedMouseButtons.Add( MouseEvent.GetEffectingButton() );
 
+	// Input preprocessor gets first chance at the input
+	if( InputPreProcessor.IsValid() && InputPreProcessor->HandleMouseButtonDownEvent( *this, MouseEvent ) )
+	{
+		return true;
+	}
+
 	bool bInGame = false;
 
 	// Only process mouse down messages if we are not drag/dropping
@@ -5490,6 +5496,12 @@ bool FSlateApplication::ProcessMouseButtonUpEvent( FPointerEvent& MouseEvent )
 	LastUserInteractionTimeForThrottling = LastUserInteractionTime;
 	PressedMouseButtons.Remove( MouseEvent.GetEffectingButton() );
 
+	// Input preprocessor gets first chance at the input
+	if( InputPreProcessor.IsValid() && InputPreProcessor->HandleMouseButtonUpEvent( *this, MouseEvent ) )
+	{
+		return true;
+	}
+
 	if ( DragDetector.DetectDragForWidget.IsValid() )
 	{
 		if ( MouseEvent.GetEffectingButton() == DragDetector.DetectDragButton && MouseEvent.GetPointerIndex() == DragDetector.DetectDragPointerIndex )
diff --git a/Engine/Source/Runtime/Slate/Public/Framework/Application/IInputProcessor.h b/Engine/Source/Runtime/Slate/Public/Framework/Application/IInputProcessor.h
--- a/Engine/Source/Runtime/Slate/Public/Framework/Application/IInputProcessor.h
+++ b/Engine/Source/Runtime/Slate/Public/Framework/Application/IInputProcessor.h
@@ -22,10 +22,25 @@ public:
 
 	virtual void Tick(const float DeltaTime, FSlateApplication& SlateApp, TSharedRef<ICursor> Cursor) = 0;
 
-	virtual bool HandleKeyDownEvent(FSlateApplication& SlateApp, const FKeyEvent& InKeyEvent) = 0;
-	virtual bool HandleKeyUpEvent(FSlateApplication& SlateApp, const FKeyEvent& InKeyEvent) = 0;
-	virtual bool HandleAnalogInputEvent(FSlateApplication& SlateApp, const FAnalogInputEvent& InAnalogInputEvent) = 0;
-	virtual bool HandleMouseMoveEvent(FSlateApplication& SlateApp, const FPointerEvent& MouseEvent) = 0;
+	/** Key down input */
+	virtual bool HandleKeyDownEvent(FSlateApplication& SlateApp, const FKeyEvent& InKeyEvent) { return false; }
+
+	/** Key up input */
+	virtual bool HandleKeyUpEvent(FSlateApplication& SlateApp, const FKeyEvent& InKeyEvent) { return false; }
+
+	/** Analog axis input */
+	virtual bool HandleAnalogInputEvent(FSlateApplication& SlateApp, const FAnalogInputEvent& InAnalogInputEvent) { return false; }
+
+	/** Mouse movement input */
+	virtual bool HandleMouseMoveEvent(FSlateApplication& SlateApp, const FPointerEvent& MouseEvent) { return false; }
+
+	/** Mouse button press */
+	virtual bool HandleMouseButtonDownEvent( FSlateApplication& SlateApp, const FPointerEvent& MouseEvent) { return false; }
+
+	/** Mouse button release */
+	virtual bool HandleMouseButtonUpEvent( FSlateApplication& SlateApp, const FPointerEvent& MouseEvent) { return false; }
+
+	/** Called when a motion-driven device has new input */
 	virtual bool HandleMotionDetectedEvent(FSlateApplication& SlateApp, const FMotionEvent& MotionEvent) { return false; };
 
 };
