ViewportInteraction: Fixed static analysis warning about dereferencing (unused) null ptr

diff --git a/Engine/Source/Editor/ViewportInteraction/Public/ViewportWorldInteraction.h b/Engine/Source/Editor/ViewportInteraction/Public/ViewportWorldInteraction.h
--- a/Engine/Source/Editor/ViewportInteraction/Public/ViewportWorldInteraction.h
+++ b/Engine/Source/Editor/ViewportInteraction/Public/ViewportWorldInteraction.h
@@ -322,7 +322,6 @@ private:
 		const EViewportInteractionDraggingMode DraggingMode, 
 		const ETransformGizmoInteractionType InteractionType, 
 		const bool bWithTwoHands, 
-		class FEditorViewportClient& ViewportClient, 
 		const TOptional<FTransformGizmoHandlePlacement> OptionalHandlePlacement, 
 		const FVector& DragDelta, 
 		const FVector& OtherHandDragDelta, 
diff --git a/Engine/Source/Editor/ViewportInteraction/ViewportWorldInteraction.cpp b/Engine/Source/Editor/ViewportInteraction/ViewportWorldInteraction.cpp
--- a/Engine/Source/Editor/ViewportInteraction/ViewportWorldInteraction.cpp
+++ b/Engine/Source/Editor/ViewportInteraction/ViewportWorldInteraction.cpp
@@ -1017,7 +1017,6 @@ void UViewportWorldInteraction::InteractionTick( const float DeltaTime )
 				InteractorData.DraggingMode, 
 				InteractorData.TransformGizmoInteractionType, 
 				bWithTwoHands,
-				*DefaultOptionalViewportClient, 
 				InteractorData.OptionalHandlePlacement, 
 				DragDelta, 
 				OtherHandDragDelta,
@@ -1147,7 +1146,6 @@ void UViewportWorldInteraction::InteractionTick( const float DeltaTime )
 					InteractorData.LastDraggingMode, 
 					InteractorData.TransformGizmoInteractionType, 
 					bWithTwoHands, 
-					*DefaultOptionalViewportClient, 
 					InteractorData.OptionalHandlePlacement, 
 					DragDelta, 
 					OtherHandDragDelta,
@@ -1337,7 +1335,6 @@ void UViewportWorldInteraction::UpdateDragging(
 	const EViewportInteractionDraggingMode DraggingMode,
 	const ETransformGizmoInteractionType InteractionType,
 	const bool bWithTwoHands,
-	FEditorViewportClient& ViewportClient,
 	const TOptional<FTransformGizmoHandlePlacement> OptionalHandlePlacement,
 	const FVector& DragDelta,
 	const FVector& OtherHandDragDelta,
