VREditor: - Transform gizmo quarter rotation handles aiming towards the users view.
- Transform gizmo added full rotation handle when dragging a rotation.
- Removed bAllHandlesVisible arguments from UGizmoHandleGroup::UpdateGizmoHandleGroup.

diff --git a/Engine/Source/Editor/VREditor/Gizmo/VREditorPlaneTranslationGizmoHandle.cpp b/Engine/Source/Editor/VREditor/Gizmo/VREditorPlaneTranslationGizmoHandle.cpp
--- a/Engine/Source/Editor/VREditor/Gizmo/VREditorPlaneTranslationGizmoHandle.cpp
+++ b/Engine/Source/Editor/VREditor/Gizmo/VREditorPlaneTranslationGizmoHandle.cpp
@@ -47,11 +47,10 @@ UVREditorPlaneTranslationGizmoHandleGroup::UVREditorPlaneTranslationGizmoHandleG
 	}
 }
 
-void UVREditorPlaneTranslationGizmoHandleGroup::UpdateGizmoHandleGroup(const FTransform& LocalToWorld, const FBox& LocalBounds, const FVector ViewLocation, bool bAllHandlesVisible, class UActorComponent* DraggingHandle,
-	const TArray< UActorComponent* >& HoveringOverHandles, float AnimationAlpha, float GizmoScale, const float GizmoHoverScale, const float GizmoHoverAnimationDuration, bool& bOutIsHoveringOrDraggingThisHandleGroup )
+void UVREditorPlaneTranslationGizmoHandleGroup::UpdateGizmoHandleGroup(const FTransform& LocalToWorld, const FBox& LocalBounds, const FVector ViewLocation, class UActorComponent* DraggingHandle, const TArray< UActorComponent* >& HoveringOverHandles, float AnimationAlpha, float GizmoScale, const float GizmoHoverScale, const float GizmoHoverAnimationDuration, bool& bOutIsHoveringOrDraggingThisHandleGroup )
 {
 	// Call parent implementation (updates hover animation)
-	Super::UpdateGizmoHandleGroup(LocalToWorld, LocalBounds, ViewLocation, bAllHandlesVisible, DraggingHandle, HoveringOverHandles, AnimationAlpha,
+	Super::UpdateGizmoHandleGroup(LocalToWorld, LocalBounds, ViewLocation, DraggingHandle, HoveringOverHandles, AnimationAlpha,
 		GizmoScale, GizmoHoverScale, GizmoHoverAnimationDuration, bOutIsHoveringOrDraggingThisHandleGroup );
 
 	for (int32 HandleIndex = 0; HandleIndex < Handles.Num(); ++HandleIndex)
diff --git a/Engine/Source/Editor/VREditor/Gizmo/VREditorPlaneTranslationGizmoHandle.h b/Engine/Source/Editor/VREditor/Gizmo/VREditorPlaneTranslationGizmoHandle.h
--- a/Engine/Source/Editor/VREditor/Gizmo/VREditorPlaneTranslationGizmoHandle.h
+++ b/Engine/Source/Editor/VREditor/Gizmo/VREditorPlaneTranslationGizmoHandle.h
@@ -23,7 +23,7 @@ public:
 	UVREditorPlaneTranslationGizmoHandleGroup();
 	
 	/** Updates the gizmo handles */
-	void UpdateGizmoHandleGroup(const FTransform& LocalToWorld, const FBox& LocalBounds, const FVector ViewLocation, bool bAllHandlesVisible, class UActorComponent* DraggingHandle,
+	void UpdateGizmoHandleGroup(const FTransform& LocalToWorld, const FBox& LocalBounds, const FVector ViewLocation, class UActorComponent* DraggingHandle,
 		const TArray< UActorComponent* >& HoveringOverHandles, float AnimationAlpha, float GizmoScale, const float GizmoHoverScale, const float GizmoHoverAnimationDuration, bool& bOutIsHoveringOrDraggingThisHandleGroup ) override;
 
 	/** Gets the InteractionType for this Gizmo handle */
diff --git a/Engine/Source/Editor/VREditor/Gizmo/VREditorRotationGizmoHandle.cpp b/Engine/Source/Editor/VREditor/Gizmo/VREditorRotationGizmoHandle.cpp
--- a/Engine/Source/Editor/VREditor/Gizmo/VREditorRotationGizmoHandle.cpp
+++ b/Engine/Source/Editor/VREditor/Gizmo/VREditorRotationGizmoHandle.cpp
@@ -47,11 +47,11 @@ UVREditorRotationGizmoHandleGroup::UVREditorRotationGizmoHandleGroup()
 	}
 }
 
-void UVREditorRotationGizmoHandleGroup::UpdateGizmoHandleGroup( const FTransform& LocalToWorld, const FBox& LocalBounds, const FVector ViewLocation, bool bAllHandlesVisible, class UActorComponent* DraggingHandle, 
+void UVREditorRotationGizmoHandleGroup::UpdateGizmoHandleGroup( const FTransform& LocalToWorld, const FBox& LocalBounds, const FVector ViewLocation, class UActorComponent* DraggingHandle, 
 	const TArray< UActorComponent* >& HoveringOverHandles, float AnimationAlpha, float GizmoScale, const float GizmoHoverScale, const float GizmoHoverAnimationDuration, bool& bOutIsHoveringOrDraggingThisHandleGroup )
 {
 	// Call parent implementation (updates hover animation)
-	Super::UpdateGizmoHandleGroup(LocalToWorld, LocalBounds, ViewLocation, bAllHandlesVisible, DraggingHandle, HoveringOverHandles,
+	Super::UpdateGizmoHandleGroup(LocalToWorld, LocalBounds, ViewLocation, DraggingHandle, HoveringOverHandles,
 		AnimationAlpha, GizmoScale, GizmoHoverScale, GizmoHoverAnimationDuration, bOutIsHoveringOrDraggingThisHandleGroup);
 
 	// Place the rotation handles
diff --git a/Engine/Source/Editor/VREditor/Gizmo/VREditorRotationGizmoHandle.h b/Engine/Source/Editor/VREditor/Gizmo/VREditorRotationGizmoHandle.h
--- a/Engine/Source/Editor/VREditor/Gizmo/VREditorRotationGizmoHandle.h
+++ b/Engine/Source/Editor/VREditor/Gizmo/VREditorRotationGizmoHandle.h
@@ -23,7 +23,7 @@ public:
 	UVREditorRotationGizmoHandleGroup();
 
 	/** Updates the Gizmo handles */
-	virtual void UpdateGizmoHandleGroup( const FTransform& LocalToWorld, const FBox& LocalBounds, const FVector ViewLocation, bool bAllHandlesVisible, class UActorComponent* DraggingHandle, const TArray< UActorComponent* >& HoveringOverHandles, 
+	virtual void UpdateGizmoHandleGroup( const FTransform& LocalToWorld, const FBox& LocalBounds, const FVector ViewLocation, class UActorComponent* DraggingHandle, const TArray< UActorComponent* >& HoveringOverHandles, 
 		float AnimationAlpha, float GizmoScale, const float GizmoHoverScale, const float GizmoHoverAnimationDuration, bool& bOutIsHoveringOrDraggingThisHandleGroup ) override;
 
 	/** Gets the InteractionType for this Gizmo handle */
diff --git a/Engine/Source/Editor/VREditor/Gizmo/VREditorTransformGizmo.cpp b/Engine/Source/Editor/VREditor/Gizmo/VREditorTransformGizmo.cpp
--- a/Engine/Source/Editor/VREditor/Gizmo/VREditorTransformGizmo.cpp
+++ b/Engine/Source/Editor/VREditor/Gizmo/VREditorTransformGizmo.cpp
@@ -179,7 +179,7 @@ void ATransformGizmo::UpdateGizmo( const EGizmoHandleTypes GizmoType, const ECoo
 		if ( HandleGroup != nullptr )
 		{
 			bool bIsHoveringOrDraggingThisHandleGroup = false;
-			HandleGroup->UpdateGizmoHandleGroup( LocalToWorld, LocalBounds, ViewLocation, bAllHandlesVisible, DraggingHandle, 
+			HandleGroup->UpdateGizmoHandleGroup( LocalToWorld, LocalBounds, ViewLocation, DraggingHandle, 
 				HoveringOverHandles, AnimationAlpha, GizmoScale, GizmoHoverScale, GizmoHoverAnimationDuration, /* Out */ bIsHoveringOrDraggingThisHandleGroup );
 			
 			if( HandleGroup->GetHandleType() == EGizmoHandleTypes::Scale && bIsHoveringOrDraggingThisHandleGroup )
diff --git a/Engine/Source/Editor/VREditor/Gizmo/VREditorTranslationGizmoHandle.cpp b/Engine/Source/Editor/VREditor/Gizmo/VREditorTranslationGizmoHandle.cpp
--- a/Engine/Source/Editor/VREditor/Gizmo/VREditorTranslationGizmoHandle.cpp
+++ b/Engine/Source/Editor/VREditor/Gizmo/VREditorTranslationGizmoHandle.cpp
@@ -47,11 +47,11 @@ UVREditorTranslationGizmoHandleGroup::UVREditorTranslationGizmoHandleGroup()
 	}
 }
 
-void UVREditorTranslationGizmoHandleGroup::UpdateGizmoHandleGroup( const FTransform& LocalToWorld, const FBox& LocalBounds, const FVector ViewLocation, bool bAllHandlesVisible, class UActorComponent* DraggingHandle, 
+void UVREditorTranslationGizmoHandleGroup::UpdateGizmoHandleGroup( const FTransform& LocalToWorld, const FBox& LocalBounds, const FVector ViewLocation, class UActorComponent* DraggingHandle, 
 	const TArray< UActorComponent* >& HoveringOverHandles, float AnimationAlpha, float GizmoScale, const float GizmoHoverScale, const float GizmoHoverAnimationDuration, bool& bOutIsHoveringOrDraggingThisHandleGroup )
 {
 	// Call parent implementation (updates hover animation)
-	Super::UpdateGizmoHandleGroup(LocalToWorld, LocalBounds, ViewLocation, bAllHandlesVisible, DraggingHandle, HoveringOverHandles,
+	Super::UpdateGizmoHandleGroup(LocalToWorld, LocalBounds, ViewLocation, DraggingHandle, HoveringOverHandles,
 		AnimationAlpha, GizmoScale, GizmoHoverScale, GizmoHoverAnimationDuration, bOutIsHoveringOrDraggingThisHandleGroup );
 
 	// Place the translation handles
diff --git a/Engine/Source/Editor/VREditor/Gizmo/VREditorTranslationGizmoHandle.h b/Engine/Source/Editor/VREditor/Gizmo/VREditorTranslationGizmoHandle.h
--- a/Engine/Source/Editor/VREditor/Gizmo/VREditorTranslationGizmoHandle.h
+++ b/Engine/Source/Editor/VREditor/Gizmo/VREditorTranslationGizmoHandle.h
@@ -23,7 +23,7 @@ public:
 	UVREditorTranslationGizmoHandleGroup();
 	
 	/** Updates the Gizmo handles */
-	virtual void UpdateGizmoHandleGroup( const FTransform& LocalToWorld, const FBox& LocalBounds, const FVector ViewLocation, bool bAllHandlesVisible, class UActorComponent* DraggingHandle,
+	virtual void UpdateGizmoHandleGroup( const FTransform& LocalToWorld, const FBox& LocalBounds, const FVector ViewLocation, class UActorComponent* DraggingHandle,
 		const TArray< UActorComponent* >& HoveringOverHandles, float AnimationAlpha, float GizmoScale, const float GizmoHoverScale, const float GizmoHoverAnimationDuration, bool& bOutIsHoveringOrDraggingThisHandleGroup ) override;
 
 	/** Gets the InteractionType for this Gizmo handle */
diff --git a/Engine/Source/Editor/ViewportInteraction/Gizmo/VIGizmoHandle.cpp b/Engine/Source/Editor/ViewportInteraction/Gizmo/VIGizmoHandle.cpp
--- a/Engine/Source/Editor/ViewportInteraction/Gizmo/VIGizmoHandle.cpp
+++ b/Engine/Source/Editor/ViewportInteraction/Gizmo/VIGizmoHandle.cpp
@@ -9,11 +9,17 @@
 
 UGizmoHandleGroup::UGizmoHandleGroup()
 	: Super(),
-	bShowOnUniversalGizmo( true )
+	bShowOnUniversalGizmo(true),
+	OwningTransformGizmoActor(nullptr)
 {
 
 }
 
+UGizmoHandleGroup::~UGizmoHandleGroup()
+{
+	OwningTransformGizmoActor = nullptr;
+}
+
 FTransformGizmoHandlePlacement UGizmoHandleGroup::MakeHandlePlacementForIndex( const int32 HandleIndex ) const
 {
 	FTransformGizmoHandlePlacement HandlePlacement;
@@ -115,7 +121,7 @@ ETransformGizmoInteractionType UGizmoHandleGroup::GetInteractionType() const
 	return ETransformGizmoInteractionType::Translate;
 }
 
-void UGizmoHandleGroup::UpdateGizmoHandleGroup(const FTransform& LocalToWorld, const FBox& LocalBounds, const FVector ViewLocation, bool bAllHandlesVisible, class UActorComponent* DraggingHandle, 
+void UGizmoHandleGroup::UpdateGizmoHandleGroup(const FTransform& LocalToWorld, const FBox& LocalBounds, const FVector ViewLocation, class UActorComponent* DraggingHandle, 
 	const TArray< UActorComponent* >& HoveringOverHandles, float AnimationAlpha, float GizmoScale, const float GizmoHoverScale, const float GizmoHoverAnimationDuration, bool& bOutIsHoveringOrDraggingThisHandleGroup )
 {
 	UpdateHoverAnimation(DraggingHandle, HoveringOverHandles, GizmoHoverAnimationDuration, bOutIsHoveringOrDraggingThisHandleGroup);
@@ -163,6 +169,11 @@ bool UGizmoHandleGroup::GetShowOnUniversalGizmo() const
 	return bShowOnUniversalGizmo;
 }
 
+void UGizmoHandleGroup::SetOwningTransformGizmo(class ABaseTransformGizmo* TransformGizmo)
+{
+	OwningTransformGizmoActor = TransformGizmo;
+}
+
 void UGizmoHandleGroup::UpdateHandleColor( const int32 AxisIndex, FGizmoHandle& Handle, class UActorComponent* DraggingHandle, const TArray< UActorComponent* >& HoveringOverHandles )
 {
 	UStaticMeshComponent* HandleMesh = Handle.HandleMesh;
@@ -221,6 +232,25 @@ void UGizmoHandleGroup::UpdateHandleColor( const int32 AxisIndex, FGizmoHandle&
 	}
 }
 
+void UGizmoHandleGroup::UpdateVisibilityAndCollision(const EGizmoHandleTypes GizmoType, const ECoordSystem GizmoCoordinateSpace, const bool bAllHandlesVisible, UActorComponent* DraggingHandle)
+{
+	const bool bIsTypeSupported = (GizmoType == EGizmoHandleTypes::All && GetShowOnUniversalGizmo()) || GetHandleType() == GizmoType;
+	const bool bSupportsCurrentCoordinateSpace = SupportsWorldCoordinateSpace() || GizmoCoordinateSpace != COORD_World;
+
+	for (FGizmoHandle& Handle : GetHandles())
+	{
+		if (Handle.HandleMesh != nullptr)
+		{
+			const bool bShowIt = (bIsTypeSupported && bSupportsCurrentCoordinateSpace && bAllHandlesVisible) || (DraggingHandle != nullptr && DraggingHandle == Handle.HandleMesh);
+
+			Handle.HandleMesh->SetVisibility(bShowIt);
+
+			// Never allow ray queries to impact hidden handles
+			Handle.HandleMesh->SetCollisionEnabled(bShowIt ? ECollisionEnabled::QueryOnly : ECollisionEnabled::NoCollision);
+		}
+	}
+}
+
 class UStaticMeshComponent* UGizmoHandleGroup::CreateMeshHandle( class UStaticMesh* HandleMesh, const FString& ComponentName )
 {
 	const bool bAllowGizmoLighting = false;	// @todo vreditor: Not sure if we want this for gizmos or not yet.  Needs feedback.  Also they're translucent right now.
diff --git a/Engine/Source/Editor/ViewportInteraction/Gizmo/VIPivotTransformGizmo.cpp b/Engine/Source/Editor/ViewportInteraction/Gizmo/VIPivotTransformGizmo.cpp
--- a/Engine/Source/Editor/ViewportInteraction/Gizmo/VIPivotTransformGizmo.cpp
+++ b/Engine/Source/Editor/ViewportInteraction/Gizmo/VIPivotTransformGizmo.cpp
@@ -7,53 +7,60 @@
 #include "GameFramework/WorldSettings.h"
 #include "VIStretchGizmoHandle.h"
 #include "VIUniformScaleGizmoHandle.h"
+#include "ViewportWorldInteraction.h"
 
 namespace VREd //@todo VREditor: Duplicates of TransformGizmo
 {
 	// @todo vreditor tweak: Tweak out console variables
 	static FAutoConsoleVariable PivotGizmoMinDistanceForScaling( TEXT( "VI.PivotGizmoMinDistanceForScaling" ), 100.0f, TEXT( "How far away the camera needs to be from an object before we'll start scaling it based on distance"));
 	static FAutoConsoleVariable PivotGizmoDistanceScaleFactor( TEXT( "VI.PivotGizmoDistanceScaleFactor" ), 0.003f, TEXT( "How much the gizmo handles should increase in size with distance from the camera, to make it easier to select" ) );
-	static FAutoConsoleVariable PivotGizmoTranslationPivotOffsetX( TEXT("VI.PivotGizmoTranslationPivotOffsetX" ), 20.0f, TEXT( "How much the translation cylinder is offsetted from the pivot" ) );
+	static FAutoConsoleVariable PivotGizmoTranslationPivotOffsetX( TEXT("VI.PivotGizmoTranslationPivotOffsetX" ), 30.0f, TEXT( "How much the translation cylinder is offsetted from the pivot" ) );
 	static FAutoConsoleVariable PivotGizmoScalePivotOffsetX( TEXT( "VI.PivotGizmoScalePivotOffsetX" ), 120.0f, TEXT( "How much the non-uniform scale is offsetted from the pivot" ) );
 	static FAutoConsoleVariable PivotGizmoPlaneTranslationPivotOffsetYZ(TEXT("VI.PivotGizmoPlaneTranslationPivotOffsetYZ" ), 40.0f, TEXT( "How much the plane translation is offsetted from the pivot" ) );
-	static FAutoConsoleVariable PivotGizmoTranslationScaleMultiply( TEXT( "VI.PivotGizmoTranslationScaleMultiply" ), 3.0f, TEXT( "Multiplies translation handles scale" ) );
-	static FAutoConsoleVariable PivotGizmoTranslationHoverScaleMultiply( TEXT( "VI.PivotGizmoTranslationHoverScaleMultiply" ), 1.2f, TEXT( "Multiplies translation handles hover scale" ) );
+	static FAutoConsoleVariable PivotGizmoTranslationScaleMultiply( TEXT( "VI.PivotGizmoTranslationScaleMultiply" ), 2.0f, TEXT( "Multiplies translation handles scale" ) );
+	static FAutoConsoleVariable PivotGizmoTranslationHoverScaleMultiply( TEXT( "VI.PivotGizmoTranslationHoverScaleMultiply" ), 1.1f, TEXT( "Multiplies translation handles hover scale" ) );
 }
 
 APivotTransformGizmo::APivotTransformGizmo() :
 	Super()
 {
 	UniformScaleGizmoHandleGroup = CreateDefaultSubobject<UUniformScaleGizmoHandleGroup>( TEXT( "UniformScaleHandles" ), true );
+	UniformScaleGizmoHandleGroup->SetOwningTransformGizmo(this);
 	UniformScaleGizmoHandleGroup->SetTranslucentGizmoMaterial( TranslucentGizmoMaterial );
 	UniformScaleGizmoHandleGroup->SetGizmoMaterial( GizmoMaterial );
 	UniformScaleGizmoHandleGroup->SetupAttachment( SceneComponent );
 	AllHandleGroups.Add( UniformScaleGizmoHandleGroup );
 
 	TranslationGizmoHandleGroup = CreateDefaultSubobject<UPivotTranslationGizmoHandleGroup>(TEXT("TranslationHandles"), true);
+	TranslationGizmoHandleGroup->SetOwningTransformGizmo(this);
 	TranslationGizmoHandleGroup->SetTranslucentGizmoMaterial(TranslucentGizmoMaterial);
 	TranslationGizmoHandleGroup->SetGizmoMaterial(GizmoMaterial);
 	TranslationGizmoHandleGroup->SetupAttachment( SceneComponent );
 	AllHandleGroups.Add(TranslationGizmoHandleGroup);
 
 	ScaleGizmoHandleGroup = CreateDefaultSubobject<UPivotScaleGizmoHandleGroup>( TEXT( "ScaleHandles" ), true );
+	ScaleGizmoHandleGroup->SetOwningTransformGizmo(this);
 	ScaleGizmoHandleGroup->SetTranslucentGizmoMaterial( TranslucentGizmoMaterial );
 	ScaleGizmoHandleGroup->SetGizmoMaterial( GizmoMaterial );
 	ScaleGizmoHandleGroup->SetupAttachment( SceneComponent );
 	AllHandleGroups.Add( ScaleGizmoHandleGroup );
 
 	PlaneTranslationGizmoHandleGroup = CreateDefaultSubobject<UPivotPlaneTranslationGizmoHandleGroup>( TEXT( "PlaneTranslationHandles" ), true );
+	PlaneTranslationGizmoHandleGroup->SetOwningTransformGizmo(this);
 	PlaneTranslationGizmoHandleGroup->SetTranslucentGizmoMaterial( TranslucentGizmoMaterial );
 	PlaneTranslationGizmoHandleGroup->SetGizmoMaterial( GizmoMaterial );
 	PlaneTranslationGizmoHandleGroup->SetupAttachment( SceneComponent );
 	AllHandleGroups.Add( PlaneTranslationGizmoHandleGroup );
 
 	RotationGizmoHandleGroup = CreateDefaultSubobject<UPivotRotationGizmoHandleGroup>( TEXT( "RotationHandles" ), true );
+	RotationGizmoHandleGroup->SetOwningTransformGizmo(this);
 	RotationGizmoHandleGroup->SetTranslucentGizmoMaterial( TranslucentGizmoMaterial );
 	RotationGizmoHandleGroup->SetGizmoMaterial( GizmoMaterial );
 	RotationGizmoHandleGroup->SetupAttachment( SceneComponent );
 	AllHandleGroups.Add( RotationGizmoHandleGroup );
 
 	StretchGizmoHandleGroup = CreateDefaultSubobject<UStretchGizmoHandleGroup>( TEXT( "StretchHandles" ), true );
+	StretchGizmoHandleGroup->SetOwningTransformGizmo(this);
 	StretchGizmoHandleGroup->SetTranslucentGizmoMaterial( TranslucentGizmoMaterial );
 	StretchGizmoHandleGroup->SetGizmoMaterial( GizmoMaterial );
 	StretchGizmoHandleGroup->SetShowOnUniversalGizmo( false );
@@ -74,6 +81,7 @@ void APivotTransformGizmo::UpdateGizmo( const EGizmoHandleTypes GizmoType, const
 	const bool bSweep = false;
 	this->SetActorTransform( LocalToWorld, bSweep );
 
+
 	// Increase scale with distance, to make gizmo handles easier to click on
 	const float WorldSpaceDistanceToToPivot = FMath::Max( VREd::PivotGizmoMinDistanceForScaling->GetFloat(), FMath::Sqrt( FVector::DistSquared( GetActorLocation(), ViewLocation ) ) );
 
@@ -87,16 +95,21 @@ void APivotTransformGizmo::UpdateGizmo( const EGizmoHandleTypes GizmoType, const
 	bool bIsHoveringOrDraggingScaleGizmo = false;
 	for ( UGizmoHandleGroup* HandleGroup : AllHandleGroups )
 	{
-		if ( HandleGroup != nullptr && 
-		   ( DraggingHandle == nullptr || HandleGroup == StretchGizmoHandleGroup ) )
+		if ( HandleGroup != nullptr)
 		{
-			bool bIsHoveringOrDraggingThisHandleGroup = false;
-			HandleGroup->UpdateGizmoHandleGroup( LocalToWorld, LocalBounds, ViewLocation, bAllHandlesVisible, DraggingHandle,
-				HoveringOverHandles, AnimationAlpha, GizmoScale, GizmoHoverScale, GizmoHoverAnimationDuration, /* Out */ bIsHoveringOrDraggingThisHandleGroup );
+			if ( DraggingHandle == nullptr || HandleGroup == StretchGizmoHandleGroup || HandleGroup == RotationGizmoHandleGroup)
+			{
+				bool bIsHoveringOrDraggingThisHandleGroup = false;
+				HandleGroup->UpdateGizmoHandleGroup( LocalToWorld, LocalBounds, ViewLocation, DraggingHandle,
+					HoveringOverHandles, AnimationAlpha, GizmoScale, GizmoHoverScale, GizmoHoverAnimationDuration, /* Out */ bIsHoveringOrDraggingThisHandleGroup );
+			}
+
+			if (HandleGroup != RotationGizmoHandleGroup)
+			{
+				HandleGroup->UpdateVisibilityAndCollision(GizmoType, GizmoCoordinateSpace, bAllHandlesVisible, DraggingHandle);
+			}
 		}
 	}
-
-	UpdateHandleVisibility( GizmoType, GizmoCoordinateSpace, bAllHandlesVisible, DraggingHandle );
 }
 
 /************************************************************************/
@@ -107,7 +120,7 @@ UPivotTranslationGizmoHandleGroup::UPivotTranslationGizmoHandleGroup() :
 {
 	UStaticMesh* TranslationHandleMesh = nullptr;
 	{
-		static ConstructorHelpers::FObjectFinder<UStaticMesh> ObjectFinder( TEXT( "/Engine/VREditor/TransformGizmo/TranslateHandleLong" ) );
+		static ConstructorHelpers::FObjectFinder<UStaticMesh> ObjectFinder( TEXT( "/Engine/VREditor/TransformGizmo/TranslateArrowHandle" ) );
 		TranslationHandleMesh = ObjectFinder.Object;
 		check( TranslationHandleMesh != nullptr );
 	}
@@ -116,11 +129,11 @@ UPivotTranslationGizmoHandleGroup::UPivotTranslationGizmoHandleGroup() :
 }
 
 
-void UPivotTranslationGizmoHandleGroup::UpdateGizmoHandleGroup( const FTransform& LocalToWorld, const FBox& LocalBounds, const FVector ViewLocation, bool bAllHandlesVisible, class UActorComponent* DraggingHandle, 
+void UPivotTranslationGizmoHandleGroup::UpdateGizmoHandleGroup( const FTransform& LocalToWorld, const FBox& LocalBounds, const FVector ViewLocation, class UActorComponent* DraggingHandle, 
 	const TArray< UActorComponent* >& HoveringOverHandles, float AnimationAlpha, float GizmoScale, const float GizmoHoverScale, const float GizmoHoverAnimationDuration, bool& bOutIsHoveringOrDraggingThisHandleGroup )
 {
 	// Call parent implementation (updates hover animation)
-	Super::UpdateGizmoHandleGroup( LocalToWorld, LocalBounds, ViewLocation, bAllHandlesVisible, DraggingHandle, HoveringOverHandles, AnimationAlpha,
+	Super::UpdateGizmoHandleGroup( LocalToWorld, LocalBounds, ViewLocation, DraggingHandle, HoveringOverHandles, AnimationAlpha,
 		GizmoScale, GizmoHoverScale, GizmoHoverAnimationDuration, bOutIsHoveringOrDraggingThisHandleGroup );
 	
 	const float MultipliedGizmoScale = GizmoScale * VREd::PivotGizmoTranslationScaleMultiply->GetFloat();
@@ -145,21 +158,21 @@ EGizmoHandleTypes UPivotTranslationGizmoHandleGroup::GetHandleType() const
 UPivotScaleGizmoHandleGroup::UPivotScaleGizmoHandleGroup() :
 	Super()
 {
-	UStaticMesh* TranslationHandleMesh = nullptr;
+	UStaticMesh* ScaleHandleMesh = nullptr;
 	{
 		static ConstructorHelpers::FObjectFinder<UStaticMesh> ObjectFinder( TEXT( "/Engine/VREditor/TransformGizmo/UniformScaleHandle" ) );
-		TranslationHandleMesh = ObjectFinder.Object;
-		check( TranslationHandleMesh != nullptr );
+		ScaleHandleMesh = ObjectFinder.Object;
+		check( ScaleHandleMesh != nullptr );
 	}
 
-	CreateHandles( TranslationHandleMesh, FString( "PivotScaleHandle" ) );	
+	CreateHandles( ScaleHandleMesh, FString( "PivotScaleHandle" ) );	
 }
 
-void UPivotScaleGizmoHandleGroup::UpdateGizmoHandleGroup( const FTransform& LocalToWorld, const FBox& LocalBounds, const FVector ViewLocation, bool bAllHandlesVisible, class UActorComponent* DraggingHandle, 
+void UPivotScaleGizmoHandleGroup::UpdateGizmoHandleGroup( const FTransform& LocalToWorld, const FBox& LocalBounds, const FVector ViewLocation, class UActorComponent* DraggingHandle, 
 	const TArray< UActorComponent* >& HoveringOverHandles, float AnimationAlpha, float GizmoScale, const float GizmoHoverScale, const float GizmoHoverAnimationDuration, bool& bOutIsHoveringOrDraggingThisHandleGroup )
 {
 	// Call parent implementation (updates hover animation)
-	Super::UpdateGizmoHandleGroup( LocalToWorld, LocalBounds, ViewLocation, bAllHandlesVisible, DraggingHandle, HoveringOverHandles, AnimationAlpha, 
+	Super::UpdateGizmoHandleGroup( LocalToWorld, LocalBounds, ViewLocation, DraggingHandle, HoveringOverHandles, AnimationAlpha, 
 		GizmoScale, GizmoHoverScale, GizmoHoverAnimationDuration, bOutIsHoveringOrDraggingThisHandleGroup );
 
 	UpdateHandlesRelativeTransformOnAxis( FTransform( FVector( VREd::PivotGizmoScalePivotOffsetX->GetFloat(), 0, 0 ) ), AnimationAlpha, GizmoScale, GizmoHoverScale, ViewLocation, DraggingHandle, HoveringOverHandles );
@@ -196,11 +209,11 @@ UPivotPlaneTranslationGizmoHandleGroup::UPivotPlaneTranslationGizmoHandleGroup()
 	CreateHandles( TranslationHandleMesh, FString( "PlaneTranslationHandle" ) );
 }
 
-void UPivotPlaneTranslationGizmoHandleGroup::UpdateGizmoHandleGroup( const FTransform& LocalToWorld, const FBox& LocalBounds, const FVector ViewLocation, bool bAllHandlesVisible, class UActorComponent* DraggingHandle,
+void UPivotPlaneTranslationGizmoHandleGroup::UpdateGizmoHandleGroup( const FTransform& LocalToWorld, const FBox& LocalBounds, const FVector ViewLocation, class UActorComponent* DraggingHandle,
 	const TArray< UActorComponent* >& HoveringOverHandles, float AnimationAlpha, float GizmoScale, const float GizmoHoverScale, const float GizmoHoverAnimationDuration, bool& bOutIsHoveringOrDraggingThisHandleGroup )
 {
 	// Call parent implementation (updates hover animation)
-	Super::UpdateGizmoHandleGroup( LocalToWorld, LocalBounds, ViewLocation, bAllHandlesVisible, DraggingHandle, HoveringOverHandles, AnimationAlpha,
+	Super::UpdateGizmoHandleGroup( LocalToWorld, LocalBounds, ViewLocation, DraggingHandle, HoveringOverHandles, AnimationAlpha,
 		GizmoScale, GizmoHoverScale, GizmoHoverAnimationDuration, bOutIsHoveringOrDraggingThisHandleGroup );
 
 	UpdateHandlesRelativeTransformOnAxis( FTransform( FVector( 0, VREd::PivotGizmoPlaneTranslationPivotOffsetYZ->GetFloat(), VREd::PivotGizmoPlaneTranslationPivotOffsetYZ->GetFloat() ) ), 
@@ -223,24 +236,126 @@ EGizmoHandleTypes UPivotPlaneTranslationGizmoHandleGroup::GetHandleType() const
 UPivotRotationGizmoHandleGroup::UPivotRotationGizmoHandleGroup() :
 	Super()
 {
-	UStaticMesh* TranslationHandleMesh = nullptr;
+	UStaticMesh* QuarterRotationHandleMesh = nullptr;
 	{
-		static ConstructorHelpers::FObjectFinder<UStaticMesh> ObjectFinder(TEXT("/Engine/VREditor/TransformGizmo/RotationHandleFull"));
-		TranslationHandleMesh = ObjectFinder.Object;
-		check(TranslationHandleMesh != nullptr);
+		static ConstructorHelpers::FObjectFinder<UStaticMesh> ObjectFinder(TEXT("/Engine/VREditor/TransformGizmo/RotationHandleQuarter"));
+		QuarterRotationHandleMesh = ObjectFinder.Object;
+		check(QuarterRotationHandleMesh != nullptr);
 	}
 
-	CreateHandles(TranslationHandleMesh, FString("RotationHandle"));
+	CreateHandles(QuarterRotationHandleMesh, FString("RotationHandle"));
+
+	{
+		static ConstructorHelpers::FObjectFinder<UStaticMesh> ObjectFinder(TEXT("/Engine/VREditor/TransformGizmo/RotationHandleFull"));
+		UStaticMesh* FullRotationHandleMesh = ObjectFinder.Object;
+		check(FullRotationHandleMesh != nullptr);
+
+		FullRotationHandleMeshComponent = CreateMeshHandle(FullRotationHandleMesh, FString("FullRotationHandle"));
+		FullRotationHandleMeshComponent->SetVisibility(false);
+		FullRotationHandleMeshComponent->SetCollisionEnabled(ECollisionEnabled::NoCollision);
+	}
 }
 
-void UPivotRotationGizmoHandleGroup::UpdateGizmoHandleGroup( const FTransform& LocalToWorld, const FBox& LocalBounds, const FVector ViewLocation, bool bAllHandlesVisible, class UActorComponent* DraggingHandle,
+void UPivotRotationGizmoHandleGroup::UpdateGizmoHandleGroup( const FTransform& LocalToWorld, const FBox& LocalBounds, const FVector ViewLocation, class UActorComponent* DraggingHandle,
 	const TArray< UActorComponent* >& HoveringOverHandles, float AnimationAlpha, float GizmoScale, const float GizmoHoverScale, const float GizmoHoverAnimationDuration, bool& bOutIsHoveringOrDraggingThisHandleGroup )
 {
-	// Call parent implementation (updates hover animation)
-	Super::UpdateGizmoHandleGroup(LocalToWorld, LocalBounds, ViewLocation, bAllHandlesVisible, DraggingHandle, HoveringOverHandles, AnimationAlpha,
-		GizmoScale, GizmoHoverScale, GizmoHoverAnimationDuration, bOutIsHoveringOrDraggingThisHandleGroup );
-	
-	UpdateHandlesRelativeTransformOnAxis( FTransform( FVector( 0, 0, 0 ) ), AnimationAlpha, GizmoScale, GizmoScale, ViewLocation, DraggingHandle, HoveringOverHandles );
+	// Call parent implementation
+	Super::UpdateGizmoHandleGroup(LocalToWorld, LocalBounds, ViewLocation, DraggingHandle, HoveringOverHandles, AnimationAlpha,
+		GizmoScale, GizmoHoverScale, GizmoHoverAnimationDuration, bOutIsHoveringOrDraggingThisHandleGroup);
+
+	bool bShowFullRotationDragHandle = false;
+	const bool bShowRotationHandles = DraggingHandle == nullptr ? true : false;
+	const ECoordSystem CoordSystem = OwningTransformGizmoActor->GetOwnerWorldInteraction()->GetTransformGizmoCoordinateSpace();
+
+	for (int32 HandleIndex = 0; HandleIndex < Handles.Num(); ++HandleIndex)
+	{
+		FGizmoHandle& Handle = Handles[HandleIndex];
+		UStaticMeshComponent* GizmoHandleMeshComponent = Handle.HandleMesh;
+		if (GizmoHandleMeshComponent != nullptr)	// Can be null if no handle for this specific placement
+		{
+			GizmoHandleMeshComponent->SetVisibility(bShowRotationHandles);
+
+			int32 CenterHandleCount, FacingAxisIndex, CenterAxisIndex;
+			const FTransformGizmoHandlePlacement HandlePlacement = MakeHandlePlacementForIndex(HandleIndex);
+			HandlePlacement.GetCenterHandleCountAndFacingAxisIndex(/* Out */ CenterHandleCount, /* Out */ FacingAxisIndex, /* Out */ CenterAxisIndex);
+
+			if (DraggingHandle != nullptr && DraggingHandle == GizmoHandleMeshComponent)
+			{
+				bShowFullRotationDragHandle = true;
+				if (!StartActorRotation.IsSet())
+				{
+					StartActorRotation = OwningTransformGizmoActor->GetTransform().GetRotation();
+				}
+
+				if (CoordSystem == ECoordSystem::COORD_Local)
+				{
+					OwningTransformGizmoActor->SetActorRotation(StartActorRotation.GetValue());
+				}
+
+				const FVector GizmoSpaceFacingAxisVector = GetAxisVector(FacingAxisIndex, HandlePlacement.Axes[FacingAxisIndex]);
+				FullRotationHandleMeshComponent->SetRelativeTransform(FTransform(GizmoSpaceFacingAxisVector.ToOrientationQuat(), FVector::ZeroVector, FVector(GizmoScale)));
+			}
+			else if(bShowRotationHandles)
+			{
+				int32 UpAxisIndex = 0;
+				int32 RightAxisIndex = 0;
+				FRotator Rotation;
+
+				if (FacingAxisIndex == 0)
+				{
+					// X, up = Z, right = Y
+					UpAxisIndex = 2;
+					RightAxisIndex = 1;
+					Rotation = FRotator::ZeroRotator;
+				}
+				else if (FacingAxisIndex == 1)
+				{
+					// Y, up = Z, right = X
+					UpAxisIndex = 2;
+					RightAxisIndex = 0;
+					Rotation = FRotator(0, -90, 0);
+
+				}
+				else if (FacingAxisIndex == 2)
+				{
+					// Z, up = Y, right = X
+					UpAxisIndex = 0;
+					RightAxisIndex = 1;
+					Rotation = FRotator(-90, 0, 0);
+				}
+
+				// Check on which side we are relative to the gizmo
+				const FVector GizmoSpaceViewLocation = GetOwner()->GetTransform().InverseTransformPosition(ViewLocation);
+				if (GizmoSpaceViewLocation[UpAxisIndex] < 0 && GizmoSpaceViewLocation[RightAxisIndex] < 0)
+				{
+					Rotation.Roll = 180;
+				}
+				else if (GizmoSpaceViewLocation[UpAxisIndex] < 0)
+				{
+					Rotation.Roll = 90;
+				}
+				else if (GizmoSpaceViewLocation[RightAxisIndex] < 0)
+				{
+					Rotation.Roll = -90;
+				}
+
+				// Set the final transform
+				GizmoHandleMeshComponent->SetRelativeTransform(FTransform(Rotation, FVector::ZeroVector, FVector(GizmoScale)));
+
+				// Update material
+				UpdateHandleColor(FacingAxisIndex, Handle, DraggingHandle, HoveringOverHandles);
+			}
+		}
+	}
+
+	// Show or hide the visuals for full rotation and set the collision accordingly
+	FullRotationHandleMeshComponent->SetCollisionEnabled(bShowFullRotationDragHandle == true ? ECollisionEnabled::QueryOnly : ECollisionEnabled::NoCollision);
+	FullRotationHandleMeshComponent->SetVisibility(bShowFullRotationDragHandle);
+
+	if (!bShowFullRotationDragHandle)
+	{
+		StartActorRotation.Reset();
+	}
 }
 
 ETransformGizmoInteractionType UPivotRotationGizmoHandleGroup::GetInteractionType() const
diff --git a/Engine/Source/Editor/ViewportInteraction/Gizmo/VIPivotTransformGizmo.h b/Engine/Source/Editor/ViewportInteraction/Gizmo/VIPivotTransformGizmo.h
--- a/Engine/Source/Editor/ViewportInteraction/Gizmo/VIPivotTransformGizmo.h
+++ b/Engine/Source/Editor/ViewportInteraction/Gizmo/VIPivotTransformGizmo.h
@@ -67,7 +67,7 @@ public:
 	UPivotTranslationGizmoHandleGroup();
 
 	/** Updates the gizmo handles */
-	virtual void UpdateGizmoHandleGroup( const FTransform& LocalToWorld, const FBox& LocalBounds, const FVector ViewLocation, bool bAllHandlesVisible, class UActorComponent* DraggingHandle, const TArray< UActorComponent* >& HoveringOverHandles, 
+	virtual void UpdateGizmoHandleGroup( const FTransform& LocalToWorld, const FBox& LocalBounds, const FVector ViewLocation, class UActorComponent* DraggingHandle, const TArray< UActorComponent* >& HoveringOverHandles, 
 		float AnimationAlpha, float GizmoScale, const float GizmoHoverScale, const float GizmoHoverAnimationDuration, bool& bOutIsHoveringOrDraggingThisHandleGroup ) override;
 
 	/** Gets the InteractionType for this Gizmo handle */
@@ -91,7 +91,7 @@ public:
 	UPivotScaleGizmoHandleGroup();
 
 	/** Updates the gizmo handles */
-	virtual void UpdateGizmoHandleGroup( const FTransform& LocalToWorld, const FBox& LocalBounds,  const FVector ViewLocation, bool bAllHandlesVisible, class UActorComponent* DraggingHandle, const TArray< UActorComponent* >& HoveringOverHandles, 
+	virtual void UpdateGizmoHandleGroup( const FTransform& LocalToWorld, const FBox& LocalBounds,  const FVector ViewLocation, class UActorComponent* DraggingHandle, const TArray< UActorComponent* >& HoveringOverHandles, 
 		float AnimationAlpha, float GizmoScale, const float GizmoHoverScale, const float GizmoHoverAnimationDuration, bool& bOutIsHoveringOrDraggingThisHandleGroup ) override;
 
 	/** Gets the InteractionType for this Gizmo handle */
@@ -118,7 +118,7 @@ public:
 	UPivotPlaneTranslationGizmoHandleGroup();
 
 	/** Updates the gizmo handles */
-	virtual void UpdateGizmoHandleGroup( const FTransform& LocalToWorld, const FBox& LocalBounds, const FVector ViewLocation, bool bAllHandlesVisible, class UActorComponent* DraggingHandle, const TArray< UActorComponent* >& HoveringOverHandles, 
+	virtual void UpdateGizmoHandleGroup( const FTransform& LocalToWorld, const FBox& LocalBounds, const FVector ViewLocation, class UActorComponent* DraggingHandle, const TArray< UActorComponent* >& HoveringOverHandles, 
 		float AnimationAlpha, float GizmoScale, const float GizmoHoverScale, const float GizmoHoverAnimationDuration, bool& bOutIsHoveringOrDraggingThisHandleGroup ) override;
 
 	/** Gets the InteractionType for this Gizmo handle */
@@ -143,7 +143,7 @@ public:
 	UPivotRotationGizmoHandleGroup();
 
 	/** Updates the gizmo handles */
-	virtual void UpdateGizmoHandleGroup(const FTransform& LocalToWorld, const FBox& LocalBounds, const FVector ViewLocation, bool bAllHandlesVisible, class UActorComponent* DraggingHandle, const TArray< UActorComponent* >& HoveringOverHandles, 
+	virtual void UpdateGizmoHandleGroup(const FTransform& LocalToWorld, const FBox& LocalBounds, const FVector ViewLocation, class UActorComponent* DraggingHandle, const TArray< UActorComponent* >& HoveringOverHandles, 
 		float AnimationAlpha, float GizmoScale, const float GizmoHoverScale, const float GizmoHoverAnimationDuration, bool& bOutIsHoveringOrDraggingThisHandleGroup ) override;
 
 	/** Gets the InteractionType for this Gizmo handle */
@@ -151,4 +151,12 @@ public:
 
 	/** Gets the GizmoType for this Gizmo handle */
 	virtual EGizmoHandleTypes GetHandleType() const override;
+
+private:
+
+	UPROPERTY()
+	UStaticMeshComponent* FullRotationHandleMeshComponent;
+
+	// The rotation of the owning transform gizmo when starting to use a rotation handle
+	TOptional<FQuat> StartActorRotation;
 };
diff --git a/Engine/Source/Editor/ViewportInteraction/Gizmo/VIStretchGizmoHandle.cpp b/Engine/Source/Editor/ViewportInteraction/Gizmo/VIStretchGizmoHandle.cpp
--- a/Engine/Source/Editor/ViewportInteraction/Gizmo/VIStretchGizmoHandle.cpp
+++ b/Engine/Source/Editor/ViewportInteraction/Gizmo/VIStretchGizmoHandle.cpp
@@ -76,11 +76,11 @@ UStretchGizmoHandleGroup::UStretchGizmoHandleGroup()
 	}
 }
 
-void UStretchGizmoHandleGroup::UpdateGizmoHandleGroup( const FTransform& LocalToWorld, const FBox& LocalBounds, const FVector ViewLocation, bool bAllHandlesVisible, class UActorComponent* DraggingHandle, const TArray< UActorComponent* >& HoveringOverHandles, 
+void UStretchGizmoHandleGroup::UpdateGizmoHandleGroup( const FTransform& LocalToWorld, const FBox& LocalBounds, const FVector ViewLocation, class UActorComponent* DraggingHandle, const TArray< UActorComponent* >& HoveringOverHandles, 
 	float AnimationAlpha, float GizmoScale, const float GizmoHoverScale, const float GizmoHoverAnimationDuration, bool& bOutIsHoveringOrDraggingThisHandleGroup )
 {
 	// Call parent implementation (updates hover animation)
-	Super::UpdateGizmoHandleGroup(LocalToWorld, LocalBounds, ViewLocation, bAllHandlesVisible, DraggingHandle, HoveringOverHandles,
+	Super::UpdateGizmoHandleGroup(LocalToWorld, LocalBounds, ViewLocation, DraggingHandle, HoveringOverHandles,
 		AnimationAlpha, GizmoScale, GizmoHoverScale, GizmoHoverAnimationDuration, bOutIsHoveringOrDraggingThisHandleGroup );
 
 	for (int32 HandleIndex = 0; HandleIndex < Handles.Num(); ++HandleIndex)
diff --git a/Engine/Source/Editor/ViewportInteraction/Gizmo/VIUniformScaleGizmoHandle.cpp b/Engine/Source/Editor/ViewportInteraction/Gizmo/VIUniformScaleGizmoHandle.cpp
--- a/Engine/Source/Editor/ViewportInteraction/Gizmo/VIUniformScaleGizmoHandle.cpp
+++ b/Engine/Source/Editor/ViewportInteraction/Gizmo/VIUniformScaleGizmoHandle.cpp
@@ -27,11 +27,11 @@ UUniformScaleGizmoHandleGroup::UUniformScaleGizmoHandleGroup()
 	NewHandle.HandleMesh = UniformScaleHandle;
 }
 
-void UUniformScaleGizmoHandleGroup::UpdateGizmoHandleGroup( const FTransform& LocalToWorld, const FBox& LocalBounds, const FVector ViewLocation, bool bAllHandlesVisible, class UActorComponent* DraggingHandle, const TArray< UActorComponent* >& HoveringOverHandles, 
+void UUniformScaleGizmoHandleGroup::UpdateGizmoHandleGroup( const FTransform& LocalToWorld, const FBox& LocalBounds, const FVector ViewLocation, class UActorComponent* DraggingHandle, const TArray< UActorComponent* >& HoveringOverHandles, 
 	float AnimationAlpha, float GizmoScale, const float GizmoHoverScale, const float GizmoHoverAnimationDuration, bool& bOutIsHoveringOrDraggingThisHandleGroup )
 {
 	// Call parent implementation (updates hover animation)
-	Super::UpdateGizmoHandleGroup( LocalToWorld, LocalBounds, ViewLocation, bAllHandlesVisible, DraggingHandle, HoveringOverHandles, 
+	Super::UpdateGizmoHandleGroup( LocalToWorld, LocalBounds, ViewLocation, DraggingHandle, HoveringOverHandles, 
 		AnimationAlpha, GizmoScale, GizmoHoverScale, GizmoHoverAnimationDuration, bOutIsHoveringOrDraggingThisHandleGroup );
 
 	FGizmoHandle& Handle = Handles[ 0 ];
diff --git a/Engine/Source/Editor/ViewportInteraction/Public/VIGizmoHandle.h b/Engine/Source/Editor/ViewportInteraction/Public/VIGizmoHandle.h
--- a/Engine/Source/Editor/ViewportInteraction/Public/VIGizmoHandle.h
+++ b/Engine/Source/Editor/ViewportInteraction/Public/VIGizmoHandle.h
@@ -38,6 +38,9 @@ public:
 	/** Default constructor that sets up CDO properties */
 	UGizmoHandleGroup();
 	
+	/** Deconstructor */
+	virtual ~UGizmoHandleGroup();
+
 	/** Given the unique index, makes a handle */
 	FTransformGizmoHandlePlacement MakeHandlePlacementForIndex( const int32 HandleIndex ) const;
 
@@ -51,9 +54,12 @@ public:
 	static FVector GetAxisVector( const int32 AxisIndex, const ETransformGizmoHandleDirection HandleDirection );
 
 	/** Updates the Gizmo handles, needs to be implemented by derived classes */
-	virtual void UpdateGizmoHandleGroup(  const FTransform& LocalToWorld, const FBox& LocalBounds, const FVector ViewLocation, bool bAllHandlesVisible, class UActorComponent* DraggingHandle, const TArray< UActorComponent* >& HoveringOverHandles, 
+	virtual void UpdateGizmoHandleGroup(  const FTransform& LocalToWorld, const FBox& LocalBounds, const FVector ViewLocation, class UActorComponent* DraggingHandle, const TArray< UActorComponent* >& HoveringOverHandles, 
 		float AnimationAlpha, float GizmoScale, const float GizmoHoverScale, const float GizmoHoverAnimationDuration, bool& bOutIsHoveringOrDraggingThisHandleGroup );
 
+	/** Default setting the visibility and collision for all the handles in this group */
+	void UpdateVisibilityAndCollision(const EGizmoHandleTypes GizmoType, const ECoordSystem GizmoCoordinateSpace, const bool bAllHandlesVisible, UActorComponent* DraggingHandle);
+
 	/** Gets the InteractionType and the HandlePlacement for this Gizmo handle */
 	virtual void GetHandleIndexInteractionType( const int32 HandleIndex, ETransformGizmoInteractionType& OutInteractionType, TOptional<FTransformGizmoHandlePlacement>& OutHandlePlacement );
 
@@ -87,6 +93,9 @@ public:
 	/** Gets if this handlegroup will be visible with the universal gizmo */
 	bool GetShowOnUniversalGizmo() const;
 
+	/** Sets the owning transform gizmo for this handle group*/
+	void SetOwningTransformGizmo(class ABaseTransformGizmo* TransformGizmo); 
+
 protected:
 	/** Updates the colors of the dynamic material instances for the handle passed using its axis index */	
 	void UpdateHandleColor( const int32 AxisIndex, FGizmoHandle& Handle, class UActorComponent* DraggingHandle, const TArray< UActorComponent* >& HoveringOverHandles );
@@ -115,6 +124,10 @@ protected:
 	UPROPERTY()
 	TArray< FGizmoHandle > Handles;
 
+	/** The actor transform gizmo owning this handlegroup */
+	UPROPERTY()
+	class ABaseTransformGizmo* OwningTransformGizmoActor;
+
 private:
 
 	/** Updates the hover animation for the HoveringOverHandles */
diff --git a/Engine/Source/Editor/ViewportInteraction/Public/VIStretchGizmoHandle.h b/Engine/Source/Editor/ViewportInteraction/Public/VIStretchGizmoHandle.h
--- a/Engine/Source/Editor/ViewportInteraction/Public/VIStretchGizmoHandle.h
+++ b/Engine/Source/Editor/ViewportInteraction/Public/VIStretchGizmoHandle.h
@@ -23,7 +23,7 @@ public:
 	UStretchGizmoHandleGroup();
 	
 	/** Updates the Gizmo handles */
-	virtual void UpdateGizmoHandleGroup( const FTransform& LocalToWorld, const FBox& LocalBounds, const FVector ViewLocation, bool bAllHandlesVisible, class UActorComponent* DraggingHandle, const TArray< UActorComponent* >& HoveringOverHandles, 
+	virtual void UpdateGizmoHandleGroup( const FTransform& LocalToWorld, const FBox& LocalBounds, const FVector ViewLocation, class UActorComponent* DraggingHandle, const TArray< UActorComponent* >& HoveringOverHandles, 
 		float AnimationAlpha, float GizmoScale, const float GizmoHoverScale, const float GizmoHoverAnimationDuration, bool& bOutIsHoveringOrDraggingThisHandleGroup ) override;
 
 	/** Gets the InteractionType for this Gizmo handle */
diff --git a/Engine/Source/Editor/ViewportInteraction/Public/VIUniformScaleGizmoHandle.h b/Engine/Source/Editor/ViewportInteraction/Public/VIUniformScaleGizmoHandle.h
--- a/Engine/Source/Editor/ViewportInteraction/Public/VIUniformScaleGizmoHandle.h
+++ b/Engine/Source/Editor/ViewportInteraction/Public/VIUniformScaleGizmoHandle.h
@@ -23,8 +23,8 @@ public:
 	UUniformScaleGizmoHandleGroup();
 	
 	/** Updates the Gizmo handles */
-	virtual void UpdateGizmoHandleGroup( const FTransform& LocalToWorld, const FBox& LocalBounds, const FVector ViewLocation, bool bAllHandlesVisible, class UActorComponent* DraggingHandle,
-		const TArray< UActorComponent* >& HoveringOverHandles, float AnimationAlpha, float GizmoScale, const float GizmoHoverScale, const float GizmoHoverAnimationDuration, bool& bOutIsHoveringOrDraggingThisHandleGroup ) override;
+	virtual void UpdateGizmoHandleGroup( const FTransform& LocalToWorld, const FBox& LocalBounds, const FVector ViewLocation, class UActorComponent* DraggingHandle, const TArray< UActorComponent* >& HoveringOverHandles, 
+		float AnimationAlpha, float GizmoScale, const float GizmoHoverScale, const float GizmoHoverAnimationDuration, bool& bOutIsHoveringOrDraggingThisHandleGroup ) override;
 
 	/** Gets the InteractionType for this Gizmo handle */
 	virtual ETransformGizmoInteractionType GetInteractionType() const override;
