Changing ViewportWorldInteraction's IsSelectableComponent check to respect bSelectable flag on primitive components

diff --git a/Engine/Source/Editor/ViewportInteraction/ViewportWorldInteraction.cpp b/Engine/Source/Editor/ViewportInteraction/ViewportWorldInteraction.cpp
--- a/Engine/Source/Editor/ViewportInteraction/ViewportWorldInteraction.cpp
+++ b/Engine/Source/Editor/ViewportInteraction/ViewportWorldInteraction.cpp
@@ -1304,7 +1304,19 @@ void UViewportWorldInteraction::InteractionTick( const float DeltaTime )
 
 bool UViewportWorldInteraction::IsInteractableComponent( const UActorComponent* Component ) const
 {
-	return true;
+	bool bResult = false;
+
+	// Don't interact primitive components that have been set as not selectable
+	if (Component != nullptr)
+	{
+		const UPrimitiveComponent* ComponentAsPrimitive = Cast<UPrimitiveComponent>(Component);
+		if (ComponentAsPrimitive != nullptr)
+		{
+			bResult = (ComponentAsPrimitive->bSelectable == true);
+		}
+	}
+	
+	return bResult;
 }
 
 ABaseTransformGizmo* UViewportWorldInteraction::GetTransformGizmoActor()
