Moving objects in VR with gizmo handles now auto-keys correctly.

diff --git a/Engine/Source/Editor/ViewportInteraction/ActorViewportTransformable.cpp b/Engine/Source/Editor/ViewportInteraction/ActorViewportTransformable.cpp
--- a/Engine/Source/Editor/ViewportInteraction/ActorViewportTransformable.cpp
+++ b/Engine/Source/Editor/ViewportInteraction/ActorViewportTransformable.cpp
@@ -11,6 +11,7 @@ void FActorViewportTransformable::ApplyTransform( const FTransform& NewTransform
 		const FTransform& ExistingTransform = Actor->GetActorTransform();
 		if( !ExistingTransform.Equals( NewTransform, 0.0f ) )
 		{
+			GEditor->BroadcastBeginObjectMovement(*Actor);
 			const bool bOnlyTranslationChanged =
 				ExistingTransform.GetRotation() == NewTransform.GetRotation() &&
 				ExistingTransform.GetScale3D() == NewTransform.GetScale3D();
@@ -22,6 +23,7 @@ void FActorViewportTransformable::ApplyTransform( const FTransform& NewTransform
 
 			const bool bFinished = false;	// @todo gizmo now: PostEditChange never called; and bFinished=true never known!!
 			Actor->PostEditMove( bFinished );
+			GEditor->BroadcastEndObjectMovement(*Actor);
 		}
 	}
 }
@@ -98,4 +100,3 @@ void FActorViewportTransformable::UpdateIgnoredActorList( TArray<AActor*>& Ignor
 		IgnoredActors.Add( Actor );
 	}
 }
-
diff --git a/Engine/Source/Editor/ViewportInteraction/ViewportWorldInteraction.cpp b/Engine/Source/Editor/ViewportInteraction/ViewportWorldInteraction.cpp
--- a/Engine/Source/Editor/ViewportInteraction/ViewportWorldInteraction.cpp
+++ b/Engine/Source/Editor/ViewportInteraction/ViewportWorldInteraction.cpp
@@ -2327,12 +2327,6 @@ void UViewportWorldInteraction::StopDragging( UViewportInteractor* Interactor )
 					  InteractorData.DraggingMode == EViewportInteractionDraggingMode::TransformablesFreely ||
 					  InteractorData.DraggingMode == EViewportInteractionDraggingMode::TransformablesAtLaserImpact )
 			{
-				// vreditor todo: Move this to where we detect a move from inertia finishes
-				FHitResult HitResult = Interactor->GetHitResultFromLaserPointer();
-				if (HitResult.Actor.IsValid())
-				{
-					GEditor->BroadcastEndObjectMovement(*HitResult.Actor);
-				}
 				// No hand is dragging the objects anymore.
 				if ( InteractorData.DraggingMode == EViewportInteractionDraggingMode::TransformablesWithGizmo &&
 					 InteractorData.TransformGizmoInteractionType == ETransformGizmoInteractionType::RotateOnAngle )
