ViewportInteraction: Made transform gizmo more compact when out of VR; minor other fixes [INTEGRATE]
- Reduced size of transform gizmo when used in desktop mode (ForceVRMode or other modes that use this)
- Reduced hover-scaled sized of new translation gizmo handles -- they were way too big!
- Fixed rotation handles not taking into account animation transition alpha when scaling

diff --git a/Engine/Source/Editor/ViewportInteraction/Gizmo/VIPivotTransformGizmo.cpp b/Engine/Source/Editor/ViewportInteraction/Gizmo/VIPivotTransformGizmo.cpp
--- a/Engine/Source/Editor/ViewportInteraction/Gizmo/VIPivotTransformGizmo.cpp
+++ b/Engine/Source/Editor/ViewportInteraction/Gizmo/VIPivotTransformGizmo.cpp
@@ -18,7 +18,7 @@ namespace VREd //@todo VREditor: Duplicates of TransformGizmo
 	static FAutoConsoleVariable PivotGizmoScalePivotOffsetX( TEXT( "VI.PivotGizmoScalePivotOffsetX" ), 120.0f, TEXT( "How much the non-uniform scale is offsetted from the pivot" ) );
 	static FAutoConsoleVariable PivotGizmoPlaneTranslationPivotOffsetYZ(TEXT("VI.PivotGizmoPlaneTranslationPivotOffsetYZ" ), 40.0f, TEXT( "How much the plane translation is offsetted from the pivot" ) );
 	static FAutoConsoleVariable PivotGizmoTranslationScaleMultiply( TEXT( "VI.PivotGizmoTranslationScaleMultiply" ), 2.0f, TEXT( "Multiplies translation handles scale" ) );
-	static FAutoConsoleVariable PivotGizmoTranslationHoverScaleMultiply( TEXT( "VI.PivotGizmoTranslationHoverScaleMultiply" ), 1.1f, TEXT( "Multiplies translation handles hover scale" ) );
+	static FAutoConsoleVariable PivotGizmoTranslationHoverScaleMultiply( TEXT( "VI.PivotGizmoTranslationHoverScaleMultiply" ), 0.75f, TEXT( "Multiplies translation handles hover scale" ) );
 }
 
 APivotTransformGizmo::APivotTransformGizmo() :
@@ -293,7 +293,7 @@ void UPivotRotationGizmoHandleGroup::UpdateGizmoHandleGroup( const FTransform& L
 				}
 
 				const FVector GizmoSpaceFacingAxisVector = GetAxisVector(FacingAxisIndex, HandlePlacement.Axes[FacingAxisIndex]);
-				FullRotationHandleMeshComponent->SetRelativeTransform(FTransform(GizmoSpaceFacingAxisVector.ToOrientationQuat(), FVector::ZeroVector, FVector(GizmoScale)));
+				FullRotationHandleMeshComponent->SetRelativeTransform(FTransform(GizmoSpaceFacingAxisVector.ToOrientationQuat(), FVector::ZeroVector, FVector(GizmoScale * AnimationAlpha)));
 			}
 			else if(bShowRotationHandles)
 			{
@@ -340,7 +340,7 @@ void UPivotRotationGizmoHandleGroup::UpdateGizmoHandleGroup( const FTransform& L
 				}
 
 				// Set the final transform
-				GizmoHandleMeshComponent->SetRelativeTransform(FTransform(Rotation, FVector::ZeroVector, FVector(GizmoScale)));
+				GizmoHandleMeshComponent->SetRelativeTransform(FTransform(Rotation, FVector::ZeroVector, FVector(GizmoScale * AnimationAlpha )));
 
 				// Update material
 				UpdateHandleColor(FacingAxisIndex, Handle, DraggingHandle, HoveringOverHandles);
diff --git a/Engine/Source/Editor/ViewportInteraction/ViewportWorldInteraction.cpp b/Engine/Source/Editor/ViewportInteraction/ViewportWorldInteraction.cpp
--- a/Engine/Source/Editor/ViewportInteraction/ViewportWorldInteraction.cpp
+++ b/Engine/Source/Editor/ViewportInteraction/ViewportWorldInteraction.cpp
@@ -242,7 +242,7 @@ UViewportWorldInteraction::UViewportWorldInteraction():
 	StartDragAngleOnRotation(),
 	StartDragHandleDirection(),
 	bShouldTransformGizmoBeVisible( true ),
-	TransformGizmoScale( 0.5f ),
+	TransformGizmoScale( 0.35f ),
 	SnapGridActor( nullptr ),
 	SnapGridMeshComponent( nullptr ),
 	SnapGridMID( nullptr ),
