ViewportInteraction: You can no longer "get close to" the transform gizmo [INTEGRATE]
- Distance-based scaling is now in effect even when closer than 1m to the gizmo origin

diff --git a/Engine/Source/Editor/ViewportInteraction/Gizmo/VIPivotTransformGizmo.cpp b/Engine/Source/Editor/ViewportInteraction/Gizmo/VIPivotTransformGizmo.cpp
--- a/Engine/Source/Editor/ViewportInteraction/Gizmo/VIPivotTransformGizmo.cpp
+++ b/Engine/Source/Editor/ViewportInteraction/Gizmo/VIPivotTransformGizmo.cpp
@@ -12,7 +12,7 @@
 namespace VREd //@todo VREditor: Duplicates of TransformGizmo
 {
 	// @todo vreditor tweak: Tweak out console variables
-	static FAutoConsoleVariable PivotGizmoMinDistanceForScaling( TEXT( "VI.PivotGizmoMinDistanceForScaling" ), 100.0f, TEXT( "How far away the camera needs to be from an object before we'll start scaling it based on distance"));
+	static FAutoConsoleVariable PivotGizmoMinDistanceForScaling( TEXT( "VI.PivotGizmoMinDistanceForScaling" ), 0.0f, TEXT( "How far away the camera needs to be from an object before we'll start scaling it based on distance"));
 	static FAutoConsoleVariable PivotGizmoDistanceScaleFactor( TEXT( "VI.PivotGizmoDistanceScaleFactor" ), 0.003f, TEXT( "How much the gizmo handles should increase in size with distance from the camera, to make it easier to select" ) );
 	static FAutoConsoleVariable PivotGizmoTranslationPivotOffsetX( TEXT("VI.PivotGizmoTranslationPivotOffsetX" ), 30.0f, TEXT( "How much the translation cylinder is offsetted from the pivot" ) );
 	static FAutoConsoleVariable PivotGizmoScalePivotOffsetX( TEXT( "VI.PivotGizmoScalePivotOffsetX" ), 120.0f, TEXT( "How much the non-uniform scale is offsetted from the pivot" ) );
