ViewportInteraction: Support hiding rotate/scale gizmo handles for single selected unoriented points [INTEGRATE]
- Unoriented points in space can't be scaled or rotated unless more than one is selected, so it doesn't make sense to show those handles (plus they get in way.)
- This will be used for a future feature

diff --git a/Engine/Source/Editor/VREditor/Gizmo/VREditorTransformGizmo.cpp b/Engine/Source/Editor/VREditor/Gizmo/VREditorTransformGizmo.cpp
--- a/Engine/Source/Editor/VREditor/Gizmo/VREditorTransformGizmo.cpp
+++ b/Engine/Source/Editor/VREditor/Gizmo/VREditorTransformGizmo.cpp
@@ -140,10 +140,10 @@ ATransformGizmo::ATransformGizmo()
 
 
 void ATransformGizmo::UpdateGizmo( const EGizmoHandleTypes GizmoType, const ECoordSystem GizmoCoordinateSpace, const FTransform& LocalToWorld, const FBox& LocalBounds, const FVector ViewLocation, const float ScaleMultiplier, bool bAllHandlesVisible,
-	UActorComponent* DraggingHandle, const TArray< UActorComponent* >& HoveringOverHandles, const float GizmoHoverScale, const float GizmoHoverAnimationDuration )
+	const bool bAllowRotationAndScaleHandles, UActorComponent* DraggingHandle, const TArray< UActorComponent* >& HoveringOverHandles, const float GizmoHoverScale, const float GizmoHoverAnimationDuration )
 {
 	Super::UpdateGizmo( GizmoType, GizmoCoordinateSpace, LocalToWorld, LocalBounds, ViewLocation, ScaleMultiplier, bAllHandlesVisible,
-		DraggingHandle, HoveringOverHandles, GizmoHoverScale, GizmoHoverAnimationDuration );
+		bAllowRotationAndScaleHandles, DraggingHandle, HoveringOverHandles, GizmoHoverScale, GizmoHoverAnimationDuration );
 
 	const float WorldScaleFactor = GetWorld()->GetWorldSettings()->WorldToMeters / 100.0f;
 
diff --git a/Engine/Source/Editor/VREditor/Gizmo/VREditorTransformGizmo.h b/Engine/Source/Editor/VREditor/Gizmo/VREditorTransformGizmo.h
--- a/Engine/Source/Editor/VREditor/Gizmo/VREditorTransformGizmo.h
+++ b/Engine/Source/Editor/VREditor/Gizmo/VREditorTransformGizmo.h
@@ -24,7 +24,7 @@ public:
 	    us to create components and set everything up nicely for the selected objects that we'll be
 		used to manipulate */
 	virtual void UpdateGizmo( const EGizmoHandleTypes GizmoType, const ECoordSystem GizmoCoordinateSpace, const FTransform& LocalToWorld, const FBox& LocalBounds, const FVector ViewLocation, const float ScaleMultiplier, bool bAllHandlesVisible,
-		class UActorComponent* DraggingHandle, const TArray< UActorComponent* >& HoveringOverHandles, const float GizmoHoverScale, const float GizmoHoverAnimationDuration ) override;
+		const bool bAllowRotationAndScaleHandles, class UActorComponent* DraggingHandle, const TArray< UActorComponent* >& HoveringOverHandles, const float GizmoHoverScale, const float GizmoHoverAnimationDuration ) override;
 
 private:
 
diff --git a/Engine/Source/Editor/ViewportInteraction/Gizmo/VIGizmoHandle.cpp b/Engine/Source/Editor/ViewportInteraction/Gizmo/VIGizmoHandle.cpp
--- a/Engine/Source/Editor/ViewportInteraction/Gizmo/VIGizmoHandle.cpp
+++ b/Engine/Source/Editor/ViewportInteraction/Gizmo/VIGizmoHandle.cpp
@@ -232,9 +232,12 @@ void UGizmoHandleGroup::UpdateHandleColor( const int32 AxisIndex, FGizmoHandle&
 	}
 }
 
-void UGizmoHandleGroup::UpdateVisibilityAndCollision(const EGizmoHandleTypes GizmoType, const ECoordSystem GizmoCoordinateSpace, const bool bAllHandlesVisible, UActorComponent* DraggingHandle)
+void UGizmoHandleGroup::UpdateVisibilityAndCollision(const EGizmoHandleTypes GizmoType, const ECoordSystem GizmoCoordinateSpace, const bool bAllHandlesVisible, const bool bAllowRotationAndScaleHandles, UActorComponent* DraggingHandle)
 {
-	const bool bIsTypeSupported = (GizmoType == EGizmoHandleTypes::All && GetShowOnUniversalGizmo()) || GetHandleType() == GizmoType;
+	const bool bIsTypeSupported = 
+		( (GizmoType == EGizmoHandleTypes::All && GetShowOnUniversalGizmo()) || GetHandleType() == GizmoType ) &&
+		( bAllowRotationAndScaleHandles || ( GetHandleType() != EGizmoHandleTypes::Rotate && GetHandleType() != EGizmoHandleTypes::Scale ) );
+
 	const bool bSupportsCurrentCoordinateSpace = SupportsWorldCoordinateSpace() || GizmoCoordinateSpace != COORD_World;
 
 	for (FGizmoHandle& Handle : GetHandles())
diff --git a/Engine/Source/Editor/ViewportInteraction/Gizmo/VIPivotTransformGizmo.cpp b/Engine/Source/Editor/ViewportInteraction/Gizmo/VIPivotTransformGizmo.cpp
--- a/Engine/Source/Editor/ViewportInteraction/Gizmo/VIPivotTransformGizmo.cpp
+++ b/Engine/Source/Editor/ViewportInteraction/Gizmo/VIPivotTransformGizmo.cpp
@@ -73,7 +73,7 @@ APivotTransformGizmo::APivotTransformGizmo() :
 }
 
 void APivotTransformGizmo::UpdateGizmo( const EGizmoHandleTypes GizmoType, const ECoordSystem GizmoCoordinateSpace, const FTransform& LocalToWorld, const FBox& LocalBounds, const FVector ViewLocation, const float ScaleMultiplier, bool bAllHandlesVisible, 
-	UActorComponent* DraggingHandle, const TArray< UActorComponent* >& HoveringOverHandles, const float GizmoHoverScale, const float GizmoHoverAnimationDuration )
+	const bool bAllowRotationAndScaleHandles, UActorComponent* DraggingHandle, const TArray< UActorComponent* >& HoveringOverHandles, const float GizmoHoverScale, const float GizmoHoverAnimationDuration )
 {
 	const float WorldScaleFactor = GetWorld()->GetWorldSettings()->WorldToMeters / 100.0f;
 
@@ -106,7 +106,7 @@ void APivotTransformGizmo::UpdateGizmo( const EGizmoHandleTypes GizmoType, const
 
 			if (HandleGroup != RotationGizmoHandleGroup)
 			{
-				HandleGroup->UpdateVisibilityAndCollision(GizmoType, GizmoCoordinateSpace, bAllHandlesVisible, DraggingHandle);
+				HandleGroup->UpdateVisibilityAndCollision(GizmoType, GizmoCoordinateSpace, bAllHandlesVisible, bAllowRotationAndScaleHandles, DraggingHandle);
 			}
 		}
 	}
diff --git a/Engine/Source/Editor/ViewportInteraction/Gizmo/VIPivotTransformGizmo.h b/Engine/Source/Editor/ViewportInteraction/Gizmo/VIPivotTransformGizmo.h
--- a/Engine/Source/Editor/ViewportInteraction/Gizmo/VIPivotTransformGizmo.h
+++ b/Engine/Source/Editor/ViewportInteraction/Gizmo/VIPivotTransformGizmo.h
@@ -25,7 +25,7 @@ public:
 	    us to create components and set everything up nicely for the selected objects that we'll be
 		used to manipulate */
 	virtual void UpdateGizmo( const EGizmoHandleTypes GizmoType, const ECoordSystem GizmoCoordinateSpace, const FTransform& LocalToWorld, const FBox& LocalBounds, const FVector ViewLocation, const float ScaleMultiplier,
-		bool bAllHandlesVisible, class UActorComponent* DraggingHandle, const TArray< UActorComponent* >& HoveringOverHandles, const float GizmoHoverScale, const float GizmoHoverAnimationDuration ) override;
+		bool bAllHandlesVisible, const bool bAllowRotationAndScaleHandles, class UActorComponent* DraggingHandle, const TArray< UActorComponent* >& HoveringOverHandles, const float GizmoHoverScale, const float GizmoHoverAnimationDuration ) override;
 
 private:
 	/** Uniform scale handle group component */
diff --git a/Engine/Source/Editor/ViewportInteraction/Public/VIBaseTransformGizmo.h b/Engine/Source/Editor/ViewportInteraction/Public/VIBaseTransformGizmo.h
--- a/Engine/Source/Editor/ViewportInteraction/Public/VIBaseTransformGizmo.h
+++ b/Engine/Source/Editor/ViewportInteraction/Public/VIBaseTransformGizmo.h
@@ -78,7 +78,7 @@ public:
 	    us to create components and set everything up nicely for the selected objects that we'll be
 		used to manipulate */
 	virtual void UpdateGizmo( const EGizmoHandleTypes GizmoType, const ECoordSystem GizmoCoordinateSpace, const FTransform& LocalToWorld, const FBox& LocalBounds, const FVector ViewLocation, const float ScaleMultiplier, bool bAllHandlesVisible,
-		class UActorComponent* DraggingHandle, const TArray< UActorComponent* >& HoveringOverHandles, const float GizmoHoverScale, const float GizmoHoverAnimationDuration ) {}
+		const bool bAllowRotationAndScaleHandles, class UActorComponent* DraggingHandle, const TArray< UActorComponent* >& HoveringOverHandles, const float GizmoHoverScale, const float GizmoHoverAnimationDuration ) {}
 
 protected:
 
diff --git a/Engine/Source/Editor/ViewportInteraction/Public/VIGizmoHandle.h b/Engine/Source/Editor/ViewportInteraction/Public/VIGizmoHandle.h
--- a/Engine/Source/Editor/ViewportInteraction/Public/VIGizmoHandle.h
+++ b/Engine/Source/Editor/ViewportInteraction/Public/VIGizmoHandle.h
@@ -58,7 +58,7 @@ public:
 		float AnimationAlpha, float GizmoScale, const float GizmoHoverScale, const float GizmoHoverAnimationDuration, bool& bOutIsHoveringOrDraggingThisHandleGroup );
 
 	/** Default setting the visibility and collision for all the handles in this group */
-	void UpdateVisibilityAndCollision(const EGizmoHandleTypes GizmoType, const ECoordSystem GizmoCoordinateSpace, const bool bAllHandlesVisible, UActorComponent* DraggingHandle);
+	void UpdateVisibilityAndCollision(const EGizmoHandleTypes GizmoType, const ECoordSystem GizmoCoordinateSpace, const bool bAllHandlesVisible, const bool bAllowRotationAndScaleHandles, UActorComponent* DraggingHandle);
 
 	/** Gets the InteractionType and the HandlePlacement for this Gizmo handle */
 	virtual void GetHandleIndexInteractionType( const int32 HandleIndex, ETransformGizmoInteractionType& OutInteractionType, TOptional<FTransformGizmoHandlePlacement>& OutHandlePlacement );
diff --git a/Engine/Source/Editor/ViewportInteraction/Public/ViewportTransformable.h b/Engine/Source/Editor/ViewportInteraction/Public/ViewportTransformable.h
--- a/Engine/Source/Editor/ViewportInteraction/Public/ViewportTransformable.h
+++ b/Engine/Source/Editor/ViewportInteraction/Public/ViewportTransformable.h
@@ -30,6 +30,13 @@ public:
 	/** Returns the bounding box of this transformable in local space */
 	virtual FBox GetLocalSpaceBoundingBox() const = 0;
 
+	/** Returns true if this transformable is a single, unoriented point in space, thus never
+	    supports being rotated or scaled when only a single transformable is selected. */
+	virtual bool IsUnorientedPoint() const
+	{
+		return false;
+	}
+
 	/** Returns true if this transformable is a physically simulated kinematic object */
 	virtual bool IsPhysicallySimulated() const
 	{
diff --git a/Engine/Source/Editor/ViewportInteraction/ViewportWorldInteraction.cpp b/Engine/Source/Editor/ViewportInteraction/ViewportWorldInteraction.cpp
--- a/Engine/Source/Editor/ViewportInteraction/ViewportWorldInteraction.cpp
+++ b/Engine/Source/Editor/ViewportInteraction/ViewportWorldInteraction.cpp
@@ -2527,6 +2527,13 @@ void UViewportWorldInteraction::RefreshTransformGizmo( const bool bNewObjectsSel
 			}
 		}
 
+		// Don't show gizmo rotation and scale handles when we're dealing with a single unoriented
+		// point in space (such as a mesh vertex), which can't reasonably be rotated or scaled.
+		const bool bHaveSingleUnorientedPoint =
+			Transformables.Num() == 1 &&
+			Transformables[ 0 ]->IsUnorientedPoint();
+		const bool bAllowRotationAndScaleHandles = !bHaveSingleUnorientedPoint;
+
 		if ( VI::ForceGizmoPivotToCenterOfObjectsBounds->GetInt() > 0 )
 		{
 			GizmoToWorld.SetLocation( GizmoToWorld.TransformPosition( GizmoSpaceSelectedObjectsBounds.GetCenter() ) );
@@ -2556,8 +2563,19 @@ void UViewportWorldInteraction::RefreshTransformGizmo( const bool bNewObjectsSel
 			ViewerLocation = GCurrentLevelEditingViewportClient->GetViewLocation();
 		}
 
-		TransformGizmoActor->UpdateGizmo( GizmoType, GizmoCoordinateSpace, GizmoToWorld, GizmoSpaceSelectedObjectsBounds, ViewerLocation, TransformGizmoScale, bAllHandlesVisible, SingleVisibleHandle,
-			HoveringOverHandles, VI::GizmoHandleHoverScale->GetFloat(), VI::GizmoHandleHoverAnimationDuration->GetFloat() );
+		TransformGizmoActor->UpdateGizmo( 
+			GizmoType, 
+			GizmoCoordinateSpace, 
+			GizmoToWorld, 
+			GizmoSpaceSelectedObjectsBounds, 
+			ViewerLocation, 
+			TransformGizmoScale, 
+			bAllHandlesVisible, 
+			bAllowRotationAndScaleHandles,
+			SingleVisibleHandle,
+			HoveringOverHandles, 
+			VI::GizmoHandleHoverScale->GetFloat(), 
+			VI::GizmoHandleHoverAnimationDuration->GetFloat() );
 
 		// Only draw if snapping is turned on
 		SpawnGridMeshActor();
