VREditor: Fix transformgizmo rotation handle visibility when dragging freely and showing up while the current gizmo type is not rotation

diff --git a/Engine/Source/Editor/VREditor/Gizmo/VREditorPlaneTranslationGizmoHandle.cpp b/Engine/Source/Editor/VREditor/Gizmo/VREditorPlaneTranslationGizmoHandle.cpp
--- a/Engine/Source/Editor/VREditor/Gizmo/VREditorPlaneTranslationGizmoHandle.cpp
+++ b/Engine/Source/Editor/VREditor/Gizmo/VREditorPlaneTranslationGizmoHandle.cpp
@@ -47,10 +47,10 @@ UVREditorPlaneTranslationGizmoHandleGroup::UVREditorPlaneTranslationGizmoHandleG
 	}
 }
 
-void UVREditorPlaneTranslationGizmoHandleGroup::UpdateGizmoHandleGroup(const FTransform& LocalToWorld, const FBox& LocalBounds, const FVector ViewLocation, class UActorComponent* DraggingHandle, const TArray< UActorComponent* >& HoveringOverHandles, float AnimationAlpha, float GizmoScale, const float GizmoHoverScale, const float GizmoHoverAnimationDuration, bool& bOutIsHoveringOrDraggingThisHandleGroup )
+void UVREditorPlaneTranslationGizmoHandleGroup::UpdateGizmoHandleGroup(const FTransform& LocalToWorld, const FBox& LocalBounds, const FVector ViewLocation, const bool bAllHandlesVisible, class UActorComponent* DraggingHandle, const TArray< UActorComponent* >& HoveringOverHandles, float AnimationAlpha, float GizmoScale, const float GizmoHoverScale, const float GizmoHoverAnimationDuration, bool& bOutIsHoveringOrDraggingThisHandleGroup )
 {
 	// Call parent implementation (updates hover animation)
-	Super::UpdateGizmoHandleGroup(LocalToWorld, LocalBounds, ViewLocation, DraggingHandle, HoveringOverHandles, AnimationAlpha,
+	Super::UpdateGizmoHandleGroup(LocalToWorld, LocalBounds, ViewLocation, bAllHandlesVisible, DraggingHandle, HoveringOverHandles, AnimationAlpha,
 		GizmoScale, GizmoHoverScale, GizmoHoverAnimationDuration, bOutIsHoveringOrDraggingThisHandleGroup );
 
 	for (int32 HandleIndex = 0; HandleIndex < Handles.Num(); ++HandleIndex)
diff --git a/Engine/Source/Editor/VREditor/Gizmo/VREditorPlaneTranslationGizmoHandle.h b/Engine/Source/Editor/VREditor/Gizmo/VREditorPlaneTranslationGizmoHandle.h
--- a/Engine/Source/Editor/VREditor/Gizmo/VREditorPlaneTranslationGizmoHandle.h
+++ b/Engine/Source/Editor/VREditor/Gizmo/VREditorPlaneTranslationGizmoHandle.h
@@ -23,7 +23,7 @@ public:
 	UVREditorPlaneTranslationGizmoHandleGroup();
 	
 	/** Updates the gizmo handles */
-	void UpdateGizmoHandleGroup(const FTransform& LocalToWorld, const FBox& LocalBounds, const FVector ViewLocation, class UActorComponent* DraggingHandle,
+	void UpdateGizmoHandleGroup(const FTransform& LocalToWorld, const FBox& LocalBounds, const FVector ViewLocation, const bool bAllHandlesVisible, class UActorComponent* DraggingHandle,
 		const TArray< UActorComponent* >& HoveringOverHandles, float AnimationAlpha, float GizmoScale, const float GizmoHoverScale, const float GizmoHoverAnimationDuration, bool& bOutIsHoveringOrDraggingThisHandleGroup ) override;
 
 	/** Gets the InteractionType for this Gizmo handle */
diff --git a/Engine/Source/Editor/VREditor/Gizmo/VREditorRotationGizmoHandle.cpp b/Engine/Source/Editor/VREditor/Gizmo/VREditorRotationGizmoHandle.cpp
--- a/Engine/Source/Editor/VREditor/Gizmo/VREditorRotationGizmoHandle.cpp
+++ b/Engine/Source/Editor/VREditor/Gizmo/VREditorRotationGizmoHandle.cpp
@@ -47,11 +47,11 @@ UVREditorRotationGizmoHandleGroup::UVREditorRotationGizmoHandleGroup()
 	}
 }
 
-void UVREditorRotationGizmoHandleGroup::UpdateGizmoHandleGroup( const FTransform& LocalToWorld, const FBox& LocalBounds, const FVector ViewLocation, class UActorComponent* DraggingHandle, 
+void UVREditorRotationGizmoHandleGroup::UpdateGizmoHandleGroup( const FTransform& LocalToWorld, const FBox& LocalBounds, const FVector ViewLocation, const bool bAllHandlesVisible, class UActorComponent* DraggingHandle, 
 	const TArray< UActorComponent* >& HoveringOverHandles, float AnimationAlpha, float GizmoScale, const float GizmoHoverScale, const float GizmoHoverAnimationDuration, bool& bOutIsHoveringOrDraggingThisHandleGroup )
 {
 	// Call parent implementation (updates hover animation)
-	Super::UpdateGizmoHandleGroup(LocalToWorld, LocalBounds, ViewLocation, DraggingHandle, HoveringOverHandles,
+	Super::UpdateGizmoHandleGroup(LocalToWorld, LocalBounds, ViewLocation, bAllHandlesVisible, DraggingHandle, HoveringOverHandles,
 		AnimationAlpha, GizmoScale, GizmoHoverScale, GizmoHoverAnimationDuration, bOutIsHoveringOrDraggingThisHandleGroup);
 
 	// Place the rotation handles
diff --git a/Engine/Source/Editor/VREditor/Gizmo/VREditorRotationGizmoHandle.h b/Engine/Source/Editor/VREditor/Gizmo/VREditorRotationGizmoHandle.h
--- a/Engine/Source/Editor/VREditor/Gizmo/VREditorRotationGizmoHandle.h
+++ b/Engine/Source/Editor/VREditor/Gizmo/VREditorRotationGizmoHandle.h
@@ -23,7 +23,7 @@ public:
 	UVREditorRotationGizmoHandleGroup();
 
 	/** Updates the Gizmo handles */
-	virtual void UpdateGizmoHandleGroup( const FTransform& LocalToWorld, const FBox& LocalBounds, const FVector ViewLocation, class UActorComponent* DraggingHandle, const TArray< UActorComponent* >& HoveringOverHandles, 
+	virtual void UpdateGizmoHandleGroup( const FTransform& LocalToWorld, const FBox& LocalBounds, const FVector ViewLocation, const bool bAllHandlesVisible, class UActorComponent* DraggingHandle, const TArray< UActorComponent* >& HoveringOverHandles, 
 		float AnimationAlpha, float GizmoScale, const float GizmoHoverScale, const float GizmoHoverAnimationDuration, bool& bOutIsHoveringOrDraggingThisHandleGroup ) override;
 
 	/** Gets the InteractionType for this Gizmo handle */
diff --git a/Engine/Source/Editor/VREditor/Gizmo/VREditorTransformGizmo.cpp b/Engine/Source/Editor/VREditor/Gizmo/VREditorTransformGizmo.cpp
--- a/Engine/Source/Editor/VREditor/Gizmo/VREditorTransformGizmo.cpp
+++ b/Engine/Source/Editor/VREditor/Gizmo/VREditorTransformGizmo.cpp
@@ -179,7 +179,7 @@ void ATransformGizmo::UpdateGizmo( const EGizmoHandleTypes GizmoType, const ECoo
 		if ( HandleGroup != nullptr )
 		{
 			bool bIsHoveringOrDraggingThisHandleGroup = false;
-			HandleGroup->UpdateGizmoHandleGroup( LocalToWorld, LocalBounds, ViewLocation, DraggingHandle, 
+			HandleGroup->UpdateGizmoHandleGroup( LocalToWorld, LocalBounds, ViewLocation, bAllHandlesVisible, DraggingHandle, 
 				HoveringOverHandles, AnimationAlpha, GizmoScale, GizmoHoverScale, GizmoHoverAnimationDuration, /* Out */ bIsHoveringOrDraggingThisHandleGroup );
 			
 			if( HandleGroup->GetHandleType() == EGizmoHandleTypes::Scale && bIsHoveringOrDraggingThisHandleGroup )
diff --git a/Engine/Source/Editor/VREditor/Gizmo/VREditorTranslationGizmoHandle.cpp b/Engine/Source/Editor/VREditor/Gizmo/VREditorTranslationGizmoHandle.cpp
--- a/Engine/Source/Editor/VREditor/Gizmo/VREditorTranslationGizmoHandle.cpp
+++ b/Engine/Source/Editor/VREditor/Gizmo/VREditorTranslationGizmoHandle.cpp
@@ -47,11 +47,11 @@ UVREditorTranslationGizmoHandleGroup::UVREditorTranslationGizmoHandleGroup()
 	}
 }
 
-void UVREditorTranslationGizmoHandleGroup::UpdateGizmoHandleGroup( const FTransform& LocalToWorld, const FBox& LocalBounds, const FVector ViewLocation, class UActorComponent* DraggingHandle, 
+void UVREditorTranslationGizmoHandleGroup::UpdateGizmoHandleGroup( const FTransform& LocalToWorld, const FBox& LocalBounds, const FVector ViewLocation, const bool bAllHandlesVisible, class UActorComponent* DraggingHandle, 
 	const TArray< UActorComponent* >& HoveringOverHandles, float AnimationAlpha, float GizmoScale, const float GizmoHoverScale, const float GizmoHoverAnimationDuration, bool& bOutIsHoveringOrDraggingThisHandleGroup )
 {
 	// Call parent implementation (updates hover animation)
-	Super::UpdateGizmoHandleGroup(LocalToWorld, LocalBounds, ViewLocation, DraggingHandle, HoveringOverHandles,
+	Super::UpdateGizmoHandleGroup(LocalToWorld, LocalBounds, ViewLocation, bAllHandlesVisible, DraggingHandle, HoveringOverHandles,
 		AnimationAlpha, GizmoScale, GizmoHoverScale, GizmoHoverAnimationDuration, bOutIsHoveringOrDraggingThisHandleGroup );
 
 	// Place the translation handles
diff --git a/Engine/Source/Editor/VREditor/Gizmo/VREditorTranslationGizmoHandle.h b/Engine/Source/Editor/VREditor/Gizmo/VREditorTranslationGizmoHandle.h
--- a/Engine/Source/Editor/VREditor/Gizmo/VREditorTranslationGizmoHandle.h
+++ b/Engine/Source/Editor/VREditor/Gizmo/VREditorTranslationGizmoHandle.h
@@ -23,7 +23,7 @@ public:
 	UVREditorTranslationGizmoHandleGroup();
 	
 	/** Updates the Gizmo handles */
-	virtual void UpdateGizmoHandleGroup( const FTransform& LocalToWorld, const FBox& LocalBounds, const FVector ViewLocation, class UActorComponent* DraggingHandle,
+	virtual void UpdateGizmoHandleGroup( const FTransform& LocalToWorld, const FBox& LocalBounds, const FVector ViewLocation, const bool bAllHandlesVisible, class UActorComponent* DraggingHandle,
 		const TArray< UActorComponent* >& HoveringOverHandles, float AnimationAlpha, float GizmoScale, const float GizmoHoverScale, const float GizmoHoverAnimationDuration, bool& bOutIsHoveringOrDraggingThisHandleGroup ) override;
 
 	/** Gets the InteractionType for this Gizmo handle */
diff --git a/Engine/Source/Editor/ViewportInteraction/Gizmo/VIBaseTransformGizmo.cpp b/Engine/Source/Editor/ViewportInteraction/Gizmo/VIBaseTransformGizmo.cpp
--- a/Engine/Source/Editor/ViewportInteraction/Gizmo/VIBaseTransformGizmo.cpp
+++ b/Engine/Source/Editor/ViewportInteraction/Gizmo/VIBaseTransformGizmo.cpp
@@ -113,6 +113,11 @@ class UViewportWorldInteraction* ABaseTransformGizmo::GetOwnerWorldInteraction()
 	return WorldInteraction;
 }
 
+EGizmoHandleTypes ABaseTransformGizmo::GetGizmoType() const
+{
+	return GizmoType;
+}
+
 void ABaseTransformGizmo::GetBoundingBoxEdge( const FBox& Box, const int32 AxisIndex, const int32 EdgeIndex, FVector& OutVertex0, FVector& OutVertex1 )
 {
 	check( AxisIndex >= 0 && AxisIndex < 3 );
diff --git a/Engine/Source/Editor/ViewportInteraction/Gizmo/VIGizmoHandle.cpp b/Engine/Source/Editor/ViewportInteraction/Gizmo/VIGizmoHandle.cpp
--- a/Engine/Source/Editor/ViewportInteraction/Gizmo/VIGizmoHandle.cpp
+++ b/Engine/Source/Editor/ViewportInteraction/Gizmo/VIGizmoHandle.cpp
@@ -121,7 +121,7 @@ ETransformGizmoInteractionType UGizmoHandleGroup::GetInteractionType() const
 	return ETransformGizmoInteractionType::Translate;
 }
 
-void UGizmoHandleGroup::UpdateGizmoHandleGroup(const FTransform& LocalToWorld, const FBox& LocalBounds, const FVector ViewLocation, class UActorComponent* DraggingHandle, 
+void UGizmoHandleGroup::UpdateGizmoHandleGroup(const FTransform& LocalToWorld, const FBox& LocalBounds, const FVector ViewLocation, const bool bAllHandlesVisible, class UActorComponent* DraggingHandle, 
 	const TArray< UActorComponent* >& HoveringOverHandles, float AnimationAlpha, float GizmoScale, const float GizmoHoverScale, const float GizmoHoverAnimationDuration, bool& bOutIsHoveringOrDraggingThisHandleGroup )
 {
 	UpdateHoverAnimation(DraggingHandle, HoveringOverHandles, GizmoHoverAnimationDuration, bOutIsHoveringOrDraggingThisHandleGroup);
diff --git a/Engine/Source/Editor/ViewportInteraction/Gizmo/VIPivotTransformGizmo.cpp b/Engine/Source/Editor/ViewportInteraction/Gizmo/VIPivotTransformGizmo.cpp
--- a/Engine/Source/Editor/ViewportInteraction/Gizmo/VIPivotTransformGizmo.cpp
+++ b/Engine/Source/Editor/ViewportInteraction/Gizmo/VIPivotTransformGizmo.cpp
@@ -72,9 +72,11 @@ APivotTransformGizmo::APivotTransformGizmo() :
 	OnNewObjectsSelected();
 }
 
-void APivotTransformGizmo::UpdateGizmo( const EGizmoHandleTypes GizmoType, const ECoordSystem GizmoCoordinateSpace, const FTransform& LocalToWorld, const FBox& LocalBounds, const FVector ViewLocation, const float ScaleMultiplier, bool bAllHandlesVisible, 
+void APivotTransformGizmo::UpdateGizmo( const EGizmoHandleTypes InGizmoType, const ECoordSystem GizmoCoordinateSpace, const FTransform& LocalToWorld, const FBox& LocalBounds, const FVector ViewLocation, const float ScaleMultiplier, bool bAllHandlesVisible, 
 	const bool bAllowRotationAndScaleHandles, UActorComponent* DraggingHandle, const TArray< UActorComponent* >& HoveringOverHandles, const float GizmoHoverScale, const float GizmoHoverAnimationDuration )
 {
+	GizmoType = InGizmoType;
+
 	const float WorldScaleFactor = GetWorld()->GetWorldSettings()->WorldToMeters / 100.0f;
 
 	// Position the gizmo at the location of the first selected actor
@@ -100,10 +102,10 @@ void APivotTransformGizmo::UpdateGizmo( const EGizmoHandleTypes GizmoType, const
 			if ( DraggingHandle == nullptr || HandleGroup == StretchGizmoHandleGroup || HandleGroup == RotationGizmoHandleGroup)
 			{
 				bool bIsHoveringOrDraggingThisHandleGroup = false;
-				HandleGroup->UpdateGizmoHandleGroup( LocalToWorld, LocalBounds, ViewLocation, DraggingHandle,
+				HandleGroup->UpdateGizmoHandleGroup( LocalToWorld, LocalBounds, ViewLocation, bAllHandlesVisible, DraggingHandle,
 					HoveringOverHandles, AnimationAlpha, GizmoScale, GizmoHoverScale, GizmoHoverAnimationDuration, /* Out */ bIsHoveringOrDraggingThisHandleGroup );
 			}
-
+			
 			if (HandleGroup != RotationGizmoHandleGroup)
 			{
 				HandleGroup->UpdateVisibilityAndCollision(GizmoType, GizmoCoordinateSpace, bAllHandlesVisible, bAllowRotationAndScaleHandles, DraggingHandle);
@@ -129,11 +131,11 @@ UPivotTranslationGizmoHandleGroup::UPivotTranslationGizmoHandleGroup() :
 }
 
 
-void UPivotTranslationGizmoHandleGroup::UpdateGizmoHandleGroup( const FTransform& LocalToWorld, const FBox& LocalBounds, const FVector ViewLocation, class UActorComponent* DraggingHandle, 
+void UPivotTranslationGizmoHandleGroup::UpdateGizmoHandleGroup( const FTransform& LocalToWorld, const FBox& LocalBounds, const FVector ViewLocation, const bool bAllHandlesVisible, class UActorComponent* DraggingHandle, 
 	const TArray< UActorComponent* >& HoveringOverHandles, float AnimationAlpha, float GizmoScale, const float GizmoHoverScale, const float GizmoHoverAnimationDuration, bool& bOutIsHoveringOrDraggingThisHandleGroup )
 {
 	// Call parent implementation (updates hover animation)
-	Super::UpdateGizmoHandleGroup( LocalToWorld, LocalBounds, ViewLocation, DraggingHandle, HoveringOverHandles, AnimationAlpha,
+	Super::UpdateGizmoHandleGroup( LocalToWorld, LocalBounds, ViewLocation, bAllHandlesVisible, DraggingHandle, HoveringOverHandles, AnimationAlpha,
 		GizmoScale, GizmoHoverScale, GizmoHoverAnimationDuration, bOutIsHoveringOrDraggingThisHandleGroup );
 	
 	const float MultipliedGizmoScale = GizmoScale * VREd::PivotGizmoTranslationScaleMultiply->GetFloat();
@@ -168,11 +170,11 @@ UPivotScaleGizmoHandleGroup::UPivotScaleGizmoHandleGroup() :
 	CreateHandles( ScaleHandleMesh, FString( "PivotScaleHandle" ) );	
 }
 
-void UPivotScaleGizmoHandleGroup::UpdateGizmoHandleGroup( const FTransform& LocalToWorld, const FBox& LocalBounds, const FVector ViewLocation, class UActorComponent* DraggingHandle, 
+void UPivotScaleGizmoHandleGroup::UpdateGizmoHandleGroup( const FTransform& LocalToWorld, const FBox& LocalBounds, const FVector ViewLocation, const bool bAllHandlesVisible, class UActorComponent* DraggingHandle, 
 	const TArray< UActorComponent* >& HoveringOverHandles, float AnimationAlpha, float GizmoScale, const float GizmoHoverScale, const float GizmoHoverAnimationDuration, bool& bOutIsHoveringOrDraggingThisHandleGroup )
 {
 	// Call parent implementation (updates hover animation)
-	Super::UpdateGizmoHandleGroup( LocalToWorld, LocalBounds, ViewLocation, DraggingHandle, HoveringOverHandles, AnimationAlpha, 
+	Super::UpdateGizmoHandleGroup( LocalToWorld, LocalBounds, ViewLocation,bAllHandlesVisible, DraggingHandle, HoveringOverHandles, AnimationAlpha, 
 		GizmoScale, GizmoHoverScale, GizmoHoverAnimationDuration, bOutIsHoveringOrDraggingThisHandleGroup );
 
 	UpdateHandlesRelativeTransformOnAxis( FTransform( FVector( VREd::PivotGizmoScalePivotOffsetX->GetFloat(), 0, 0 ) ), AnimationAlpha, GizmoScale, GizmoHoverScale, ViewLocation, DraggingHandle, HoveringOverHandles );
@@ -209,11 +211,11 @@ UPivotPlaneTranslationGizmoHandleGroup::UPivotPlaneTranslationGizmoHandleGroup()
 	CreateHandles( TranslationHandleMesh, FString( "PlaneTranslationHandle" ) );
 }
 
-void UPivotPlaneTranslationGizmoHandleGroup::UpdateGizmoHandleGroup( const FTransform& LocalToWorld, const FBox& LocalBounds, const FVector ViewLocation, class UActorComponent* DraggingHandle,
+void UPivotPlaneTranslationGizmoHandleGroup::UpdateGizmoHandleGroup( const FTransform& LocalToWorld, const FBox& LocalBounds, const FVector ViewLocation, const bool bAllHandlesVisible, class UActorComponent* DraggingHandle,
 	const TArray< UActorComponent* >& HoveringOverHandles, float AnimationAlpha, float GizmoScale, const float GizmoHoverScale, const float GizmoHoverAnimationDuration, bool& bOutIsHoveringOrDraggingThisHandleGroup )
 {
 	// Call parent implementation (updates hover animation)
-	Super::UpdateGizmoHandleGroup( LocalToWorld, LocalBounds, ViewLocation, DraggingHandle, HoveringOverHandles, AnimationAlpha,
+	Super::UpdateGizmoHandleGroup( LocalToWorld, LocalBounds, ViewLocation, bAllHandlesVisible, DraggingHandle, HoveringOverHandles, AnimationAlpha,
 		GizmoScale, GizmoHoverScale, GizmoHoverAnimationDuration, bOutIsHoveringOrDraggingThisHandleGroup );
 
 	UpdateHandlesRelativeTransformOnAxis( FTransform( FVector( 0, VREd::PivotGizmoPlaneTranslationPivotOffsetYZ->GetFloat(), VREd::PivotGizmoPlaneTranslationPivotOffsetYZ->GetFloat() ) ), 
@@ -256,16 +258,20 @@ UPivotRotationGizmoHandleGroup::UPivotRotationGizmoHandleGroup() :
 	}
 }
 
-void UPivotRotationGizmoHandleGroup::UpdateGizmoHandleGroup( const FTransform& LocalToWorld, const FBox& LocalBounds, const FVector ViewLocation, class UActorComponent* DraggingHandle,
+void UPivotRotationGizmoHandleGroup::UpdateGizmoHandleGroup( const FTransform& LocalToWorld, const FBox& LocalBounds, const FVector ViewLocation, const bool bAllHandlesVisible, class UActorComponent* DraggingHandle,
 	const TArray< UActorComponent* >& HoveringOverHandles, float AnimationAlpha, float GizmoScale, const float GizmoHoverScale, const float GizmoHoverAnimationDuration, bool& bOutIsHoveringOrDraggingThisHandleGroup )
 {
 	// Call parent implementation
-	Super::UpdateGizmoHandleGroup(LocalToWorld, LocalBounds, ViewLocation, DraggingHandle, HoveringOverHandles, AnimationAlpha,
+	Super::UpdateGizmoHandleGroup(LocalToWorld, LocalBounds, ViewLocation, bAllHandlesVisible, DraggingHandle, HoveringOverHandles, AnimationAlpha,
 		GizmoScale, GizmoHoverScale, GizmoHoverAnimationDuration, bOutIsHoveringOrDraggingThisHandleGroup);
 
 	bool bShowFullRotationDragHandle = false;
-	const bool bShowRotationHandles = DraggingHandle == nullptr ? true : false;
 	const ECoordSystem CoordSystem = OwningTransformGizmoActor->GetOwnerWorldInteraction()->GetTransformGizmoCoordinateSpace();
+	
+	const bool bIsDraggingHandle = DraggingHandle == nullptr ? true : false;
+	const bool bIsTypeSupported = (OwningTransformGizmoActor->GetGizmoType() == EGizmoHandleTypes::All && GetShowOnUniversalGizmo()) || GetHandleType() == OwningTransformGizmoActor->GetGizmoType();
+	const bool bSupportsCurrentCoordinateSpace = SupportsWorldCoordinateSpace() || CoordSystem != COORD_World;
+	const bool bShowAnyRotationHandle = (bIsTypeSupported && bSupportsCurrentCoordinateSpace && bAllHandlesVisible);
 
 	for (int32 HandleIndex = 0; HandleIndex < Handles.Num(); ++HandleIndex)
 	{
@@ -273,13 +279,14 @@ void UPivotRotationGizmoHandleGroup::UpdateGizmoHandleGroup( const FTransform& L
 		UStaticMeshComponent* GizmoHandleMeshComponent = Handle.HandleMesh;
 		if (GizmoHandleMeshComponent != nullptr)	// Can be null if no handle for this specific placement
 		{
-			GizmoHandleMeshComponent->SetVisibility(bShowRotationHandles);
+			const bool bDraggingCurrentHandle = (DraggingHandle != nullptr && DraggingHandle == GizmoHandleMeshComponent);
+			GizmoHandleMeshComponent->SetVisibility(bShowAnyRotationHandle && !bDraggingCurrentHandle);
 
 			int32 CenterHandleCount, FacingAxisIndex, CenterAxisIndex;
 			const FTransformGizmoHandlePlacement HandlePlacement = MakeHandlePlacementForIndex(HandleIndex);
 			HandlePlacement.GetCenterHandleCountAndFacingAxisIndex(/* Out */ CenterHandleCount, /* Out */ FacingAxisIndex, /* Out */ CenterAxisIndex);
 
-			if (DraggingHandle != nullptr && DraggingHandle == GizmoHandleMeshComponent)
+			if (bDraggingCurrentHandle)
 			{
 				bShowFullRotationDragHandle = true;
 				if (!StartActorRotation.IsSet())
@@ -295,7 +302,7 @@ void UPivotRotationGizmoHandleGroup::UpdateGizmoHandleGroup( const FTransform& L
 				const FVector GizmoSpaceFacingAxisVector = GetAxisVector(FacingAxisIndex, HandlePlacement.Axes[FacingAxisIndex]);
 				FullRotationHandleMeshComponent->SetRelativeTransform(FTransform(GizmoSpaceFacingAxisVector.ToOrientationQuat(), FVector::ZeroVector, FVector(GizmoScale * AnimationAlpha)));
 			}
-			else if(bShowRotationHandles)
+			else if(DraggingHandle == nullptr)
 			{
 				int32 UpAxisIndex = 0;
 				int32 RightAxisIndex = 0;
diff --git a/Engine/Source/Editor/ViewportInteraction/Gizmo/VIPivotTransformGizmo.h b/Engine/Source/Editor/ViewportInteraction/Gizmo/VIPivotTransformGizmo.h
--- a/Engine/Source/Editor/ViewportInteraction/Gizmo/VIPivotTransformGizmo.h
+++ b/Engine/Source/Editor/ViewportInteraction/Gizmo/VIPivotTransformGizmo.h
@@ -24,7 +24,7 @@ public:
 	/** Called by the world interaction system after we've been spawned into the world, to allow
 	    us to create components and set everything up nicely for the selected objects that we'll be
 		used to manipulate */
-	virtual void UpdateGizmo( const EGizmoHandleTypes GizmoType, const ECoordSystem GizmoCoordinateSpace, const FTransform& LocalToWorld, const FBox& LocalBounds, const FVector ViewLocation, const float ScaleMultiplier,
+	virtual void UpdateGizmo( const EGizmoHandleTypes InGizmoType, const ECoordSystem GizmoCoordinateSpace, const FTransform& LocalToWorld, const FBox& LocalBounds, const FVector ViewLocation, const float ScaleMultiplier,
 		bool bAllHandlesVisible, const bool bAllowRotationAndScaleHandles, class UActorComponent* DraggingHandle, const TArray< UActorComponent* >& HoveringOverHandles, const float GizmoHoverScale, const float GizmoHoverAnimationDuration ) override;
 
 private:
@@ -67,7 +67,7 @@ public:
 	UPivotTranslationGizmoHandleGroup();
 
 	/** Updates the gizmo handles */
-	virtual void UpdateGizmoHandleGroup( const FTransform& LocalToWorld, const FBox& LocalBounds, const FVector ViewLocation, class UActorComponent* DraggingHandle, const TArray< UActorComponent* >& HoveringOverHandles, 
+	virtual void UpdateGizmoHandleGroup( const FTransform& LocalToWorld, const FBox& LocalBounds, const FVector ViewLocation, const bool bAllHandlesVisible, class UActorComponent* DraggingHandle, const TArray< UActorComponent* >& HoveringOverHandles, 
 		float AnimationAlpha, float GizmoScale, const float GizmoHoverScale, const float GizmoHoverAnimationDuration, bool& bOutIsHoveringOrDraggingThisHandleGroup ) override;
 
 	/** Gets the InteractionType for this Gizmo handle */
@@ -91,7 +91,7 @@ public:
 	UPivotScaleGizmoHandleGroup();
 
 	/** Updates the gizmo handles */
-	virtual void UpdateGizmoHandleGroup( const FTransform& LocalToWorld, const FBox& LocalBounds,  const FVector ViewLocation, class UActorComponent* DraggingHandle, const TArray< UActorComponent* >& HoveringOverHandles, 
+	virtual void UpdateGizmoHandleGroup( const FTransform& LocalToWorld, const FBox& LocalBounds,  const FVector ViewLocation, const bool bAllHandlesVisible, class UActorComponent* DraggingHandle, const TArray< UActorComponent* >& HoveringOverHandles, 
 		float AnimationAlpha, float GizmoScale, const float GizmoHoverScale, const float GizmoHoverAnimationDuration, bool& bOutIsHoveringOrDraggingThisHandleGroup ) override;
 
 	/** Gets the InteractionType for this Gizmo handle */
@@ -118,7 +118,7 @@ public:
 	UPivotPlaneTranslationGizmoHandleGroup();
 
 	/** Updates the gizmo handles */
-	virtual void UpdateGizmoHandleGroup( const FTransform& LocalToWorld, const FBox& LocalBounds, const FVector ViewLocation, class UActorComponent* DraggingHandle, const TArray< UActorComponent* >& HoveringOverHandles, 
+	virtual void UpdateGizmoHandleGroup( const FTransform& LocalToWorld, const FBox& LocalBounds, const FVector ViewLocation, const bool bAllHandlesVisible, class UActorComponent* DraggingHandle, const TArray< UActorComponent* >& HoveringOverHandles, 
 		float AnimationAlpha, float GizmoScale, const float GizmoHoverScale, const float GizmoHoverAnimationDuration, bool& bOutIsHoveringOrDraggingThisHandleGroup ) override;
 
 	/** Gets the InteractionType for this Gizmo handle */
@@ -143,7 +143,7 @@ public:
 	UPivotRotationGizmoHandleGroup();
 
 	/** Updates the gizmo handles */
-	virtual void UpdateGizmoHandleGroup(const FTransform& LocalToWorld, const FBox& LocalBounds, const FVector ViewLocation, class UActorComponent* DraggingHandle, const TArray< UActorComponent* >& HoveringOverHandles, 
+	virtual void UpdateGizmoHandleGroup(const FTransform& LocalToWorld, const FBox& LocalBounds, const FVector ViewLocation, const bool bAllHandlesVisible, class UActorComponent* DraggingHandle, const TArray< UActorComponent* >& HoveringOverHandles, 
 		float AnimationAlpha, float GizmoScale, const float GizmoHoverScale, const float GizmoHoverAnimationDuration, bool& bOutIsHoveringOrDraggingThisHandleGroup ) override;
 
 	/** Gets the InteractionType for this Gizmo handle */
diff --git a/Engine/Source/Editor/ViewportInteraction/Gizmo/VIStretchGizmoHandle.cpp b/Engine/Source/Editor/ViewportInteraction/Gizmo/VIStretchGizmoHandle.cpp
--- a/Engine/Source/Editor/ViewportInteraction/Gizmo/VIStretchGizmoHandle.cpp
+++ b/Engine/Source/Editor/ViewportInteraction/Gizmo/VIStretchGizmoHandle.cpp
@@ -76,11 +76,11 @@ UStretchGizmoHandleGroup::UStretchGizmoHandleGroup()
 	}
 }
 
-void UStretchGizmoHandleGroup::UpdateGizmoHandleGroup( const FTransform& LocalToWorld, const FBox& LocalBounds, const FVector ViewLocation, class UActorComponent* DraggingHandle, const TArray< UActorComponent* >& HoveringOverHandles, 
+void UStretchGizmoHandleGroup::UpdateGizmoHandleGroup( const FTransform& LocalToWorld, const FBox& LocalBounds, const FVector ViewLocation, const bool bAllHandlesVisible, class UActorComponent* DraggingHandle, const TArray< UActorComponent* >& HoveringOverHandles, 
 	float AnimationAlpha, float GizmoScale, const float GizmoHoverScale, const float GizmoHoverAnimationDuration, bool& bOutIsHoveringOrDraggingThisHandleGroup )
 {
 	// Call parent implementation (updates hover animation)
-	Super::UpdateGizmoHandleGroup(LocalToWorld, LocalBounds, ViewLocation, DraggingHandle, HoveringOverHandles,
+	Super::UpdateGizmoHandleGroup(LocalToWorld, LocalBounds, ViewLocation, bAllHandlesVisible, DraggingHandle, HoveringOverHandles,
 		AnimationAlpha, GizmoScale, GizmoHoverScale, GizmoHoverAnimationDuration, bOutIsHoveringOrDraggingThisHandleGroup );
 
 	for (int32 HandleIndex = 0; HandleIndex < Handles.Num(); ++HandleIndex)
diff --git a/Engine/Source/Editor/ViewportInteraction/Gizmo/VIUniformScaleGizmoHandle.cpp b/Engine/Source/Editor/ViewportInteraction/Gizmo/VIUniformScaleGizmoHandle.cpp
--- a/Engine/Source/Editor/ViewportInteraction/Gizmo/VIUniformScaleGizmoHandle.cpp
+++ b/Engine/Source/Editor/ViewportInteraction/Gizmo/VIUniformScaleGizmoHandle.cpp
@@ -27,11 +27,11 @@ UUniformScaleGizmoHandleGroup::UUniformScaleGizmoHandleGroup()
 	NewHandle.HandleMesh = UniformScaleHandle;
 }
 
-void UUniformScaleGizmoHandleGroup::UpdateGizmoHandleGroup( const FTransform& LocalToWorld, const FBox& LocalBounds, const FVector ViewLocation, class UActorComponent* DraggingHandle, const TArray< UActorComponent* >& HoveringOverHandles, 
+void UUniformScaleGizmoHandleGroup::UpdateGizmoHandleGroup( const FTransform& LocalToWorld, const FBox& LocalBounds, const FVector ViewLocation, const bool bAllHandlesVisible, class UActorComponent* DraggingHandle, const TArray< UActorComponent* >& HoveringOverHandles, 
 	float AnimationAlpha, float GizmoScale, const float GizmoHoverScale, const float GizmoHoverAnimationDuration, bool& bOutIsHoveringOrDraggingThisHandleGroup )
 {
 	// Call parent implementation (updates hover animation)
-	Super::UpdateGizmoHandleGroup( LocalToWorld, LocalBounds, ViewLocation, DraggingHandle, HoveringOverHandles, 
+	Super::UpdateGizmoHandleGroup( LocalToWorld, LocalBounds, ViewLocation, bAllHandlesVisible, DraggingHandle, HoveringOverHandles, 
 		AnimationAlpha, GizmoScale, GizmoHoverScale, GizmoHoverAnimationDuration, bOutIsHoveringOrDraggingThisHandleGroup );
 
 	FGizmoHandle& Handle = Handles[ 0 ];
diff --git a/Engine/Source/Editor/ViewportInteraction/Public/VIBaseTransformGizmo.h b/Engine/Source/Editor/ViewportInteraction/Public/VIBaseTransformGizmo.h
--- a/Engine/Source/Editor/ViewportInteraction/Public/VIBaseTransformGizmo.h
+++ b/Engine/Source/Editor/ViewportInteraction/Public/VIBaseTransformGizmo.h
@@ -80,6 +80,10 @@ public:
 	virtual void UpdateGizmo( const EGizmoHandleTypes GizmoType, const ECoordSystem GizmoCoordinateSpace, const FTransform& LocalToWorld, const FBox& LocalBounds, const FVector ViewLocation, const float ScaleMultiplier, bool bAllHandlesVisible,
 		const bool bAllowRotationAndScaleHandles, class UActorComponent* DraggingHandle, const TArray< UActorComponent* >& HoveringOverHandles, const float GizmoHoverScale, const float GizmoHoverAnimationDuration ) {}
 
+
+	/** Gets the current gizmo handle type */
+	EGizmoHandleTypes GetGizmoType() const;
+
 protected:
 
 	/** Static: Given a bounding box and information about which edge to query, returns the vertex positions of that edge */
@@ -113,4 +117,7 @@ protected:
 	/** Owning object */
 	UPROPERTY()
 	class UViewportWorldInteraction* WorldInteraction;
+	
+	/** Current gizmo type */
+	EGizmoHandleTypes GizmoType;
 };	
diff --git a/Engine/Source/Editor/ViewportInteraction/Public/VIGizmoHandle.h b/Engine/Source/Editor/ViewportInteraction/Public/VIGizmoHandle.h
--- a/Engine/Source/Editor/ViewportInteraction/Public/VIGizmoHandle.h
+++ b/Engine/Source/Editor/ViewportInteraction/Public/VIGizmoHandle.h
@@ -54,7 +54,7 @@ public:
 	static FVector GetAxisVector( const int32 AxisIndex, const ETransformGizmoHandleDirection HandleDirection );
 
 	/** Updates the Gizmo handles, needs to be implemented by derived classes */
-	virtual void UpdateGizmoHandleGroup(  const FTransform& LocalToWorld, const FBox& LocalBounds, const FVector ViewLocation, class UActorComponent* DraggingHandle, const TArray< UActorComponent* >& HoveringOverHandles, 
+	virtual void UpdateGizmoHandleGroup( const FTransform& LocalToWorld, const FBox& LocalBounds, const FVector ViewLocation, const bool bAllHandlesVisible, class UActorComponent* DraggingHandle, const TArray< UActorComponent* >& HoveringOverHandles, 
 		float AnimationAlpha, float GizmoScale, const float GizmoHoverScale, const float GizmoHoverAnimationDuration, bool& bOutIsHoveringOrDraggingThisHandleGroup );
 
 	/** Default setting the visibility and collision for all the handles in this group */
diff --git a/Engine/Source/Editor/ViewportInteraction/Public/VIStretchGizmoHandle.h b/Engine/Source/Editor/ViewportInteraction/Public/VIStretchGizmoHandle.h
--- a/Engine/Source/Editor/ViewportInteraction/Public/VIStretchGizmoHandle.h
+++ b/Engine/Source/Editor/ViewportInteraction/Public/VIStretchGizmoHandle.h
@@ -23,7 +23,7 @@ public:
 	UStretchGizmoHandleGroup();
 	
 	/** Updates the Gizmo handles */
-	virtual void UpdateGizmoHandleGroup( const FTransform& LocalToWorld, const FBox& LocalBounds, const FVector ViewLocation, class UActorComponent* DraggingHandle, const TArray< UActorComponent* >& HoveringOverHandles, 
+	virtual void UpdateGizmoHandleGroup( const FTransform& LocalToWorld, const FBox& LocalBounds, const FVector ViewLocation, const bool bAllHandlesVisible, class UActorComponent* DraggingHandle, const TArray< UActorComponent* >& HoveringOverHandles, 
 		float AnimationAlpha, float GizmoScale, const float GizmoHoverScale, const float GizmoHoverAnimationDuration, bool& bOutIsHoveringOrDraggingThisHandleGroup ) override;
 
 	/** Gets the InteractionType for this Gizmo handle */
diff --git a/Engine/Source/Editor/ViewportInteraction/Public/VIUniformScaleGizmoHandle.h b/Engine/Source/Editor/ViewportInteraction/Public/VIUniformScaleGizmoHandle.h
--- a/Engine/Source/Editor/ViewportInteraction/Public/VIUniformScaleGizmoHandle.h
+++ b/Engine/Source/Editor/ViewportInteraction/Public/VIUniformScaleGizmoHandle.h
@@ -23,7 +23,7 @@ public:
 	UUniformScaleGizmoHandleGroup();
 	
 	/** Updates the Gizmo handles */
-	virtual void UpdateGizmoHandleGroup( const FTransform& LocalToWorld, const FBox& LocalBounds, const FVector ViewLocation, class UActorComponent* DraggingHandle, const TArray< UActorComponent* >& HoveringOverHandles, 
+	virtual void UpdateGizmoHandleGroup( const FTransform& LocalToWorld, const FBox& LocalBounds, const FVector ViewLocation, const bool bAllHandlesVisible, class UActorComponent* DraggingHandle, const TArray< UActorComponent* >& HoveringOverHandles, 
 		float AnimationAlpha, float GizmoScale, const float GizmoHoverScale, const float GizmoHoverAnimationDuration, bool& bOutIsHoveringOrDraggingThisHandleGroup ) override;
 
 	/** Gets the InteractionType for this Gizmo handle */
