VREditor: Fix mac compile warning: declaration shadows a field in VIBaseTransformGizmo

diff --git a/Engine/Source/Editor/ViewportInteraction/Gizmo/VIBaseTransformGizmo.cpp b/Engine/Source/Editor/ViewportInteraction/Gizmo/VIBaseTransformGizmo.cpp
--- a/Engine/Source/Editor/ViewportInteraction/Gizmo/VIBaseTransformGizmo.cpp
+++ b/Engine/Source/Editor/ViewportInteraction/Gizmo/VIBaseTransformGizmo.cpp
@@ -221,20 +221,20 @@ void ABaseTransformGizmo::GetBoundingBoxEdge( const FBox& Box, const int32 AxisI
 
 
 
-void ABaseTransformGizmo::UpdateHandleVisibility( const EGizmoHandleTypes GizmoType, const ECoordSystem GizmoCoordinateSpace, const bool bAllHandlesVisible, UActorComponent* DraggingHandle )
+void ABaseTransformGizmo::UpdateHandleVisibility(const EGizmoHandleTypes InGizmoType, const ECoordSystem InGizmoCoordinateSpace, const bool bInAllHandlesVisible, UActorComponent* DraggingHandle)
 {
 	for ( UGizmoHandleGroup* HandleGroup : AllHandleGroups )
 	{
 		if ( HandleGroup != nullptr )
 		{
-			const bool bIsTypeSupported = ( GizmoType == EGizmoHandleTypes::All && HandleGroup->GetShowOnUniversalGizmo() ) || HandleGroup->GetHandleType() == GizmoType;
-			const bool bSupportsCurrentCoordinateSpace = HandleGroup->SupportsWorldCoordinateSpace() || GizmoCoordinateSpace != COORD_World;
+			const bool bIsTypeSupported = ( InGizmoType == EGizmoHandleTypes::All && HandleGroup->GetShowOnUniversalGizmo() ) || HandleGroup->GetHandleType() == InGizmoType;
+			const bool bSupportsCurrentCoordinateSpace = HandleGroup->SupportsWorldCoordinateSpace() || InGizmoCoordinateSpace != COORD_World;
 
 			for ( FGizmoHandle& Handle : HandleGroup->GetHandles() )
 			{
 				if( Handle.HandleMesh != nullptr )
 				{
-					const bool bShowIt = ( bIsTypeSupported && bSupportsCurrentCoordinateSpace && bAllHandlesVisible ) || ( DraggingHandle != nullptr && DraggingHandle == Handle.HandleMesh );
+					const bool bShowIt = ( bIsTypeSupported && bSupportsCurrentCoordinateSpace && bInAllHandlesVisible ) || ( DraggingHandle != nullptr && DraggingHandle == Handle.HandleMesh );
 
 					Handle.HandleMesh->SetVisibility( bShowIt );
 
diff --git a/Engine/Source/Editor/ViewportInteraction/Public/VIBaseTransformGizmo.h b/Engine/Source/Editor/ViewportInteraction/Public/VIBaseTransformGizmo.h
--- a/Engine/Source/Editor/ViewportInteraction/Public/VIBaseTransformGizmo.h
+++ b/Engine/Source/Editor/ViewportInteraction/Public/VIBaseTransformGizmo.h
@@ -90,7 +90,7 @@ protected:
 	static void GetBoundingBoxEdge( const FBox& Box, const int32 AxisIndex, const int32 EdgeIndex, FVector& OutVertex0, FVector& OutVertex1 );
 
 	/** Updates the visibility of all the handles */
-	void UpdateHandleVisibility( const EGizmoHandleTypes GizmoType, const ECoordSystem GizmoCoordinateSpace, bool bAllHandlesVisible, UActorComponent* DraggingHandle );
+	void UpdateHandleVisibility(const EGizmoHandleTypes InGizmoType, const ECoordSystem InGizmoCoordinateSpace, bool bInAllHandlesVisible, UActorComponent* DraggingHandle);
 
 	/** Gets if the gizmo shows measurement texts */
 	bool GetShowMeasurementText() const;
