Merging //UE4/Private-Geometry to Dev-VREditor (//UE4/Dev-VREditor) 3339870

diff --git a/Engine/Source/Editor/ViewportInteraction/Public/ViewportWorldInteraction.h b/Engine/Source/Editor/ViewportInteraction/Public/ViewportWorldInteraction.h
--- a/Engine/Source/Editor/ViewportInteraction/Public/ViewportWorldInteraction.h
+++ b/Engine/Source/Editor/ViewportInteraction/Public/ViewportWorldInteraction.h
@@ -16,8 +16,10 @@
 #include "VIBaseTransformGizmo.h"
 #include "LevelEditorViewport.h"
 #include "ViewportTransformable.h"
+#include "GenericPlatform/ICursor.h"
 #include "ViewportWorldInteraction.generated.h"
 
+
 class IViewportInteractableInterface;
 class UViewportInteractor;
 
@@ -26,7 +28,7 @@ namespace ViewportWorldActionTypes
 	static const FName NoAction( "NoAction" );
 	static const FName WorldMovement( "WorldMovement" );
 	static const FName SelectAndMove( "SelectAndMove" );
-	static const FName SelectAndMove_LightlyPressed( "SelectAndMove_LightlyPressed" );
+	static const FName SelectAndMove_FullyPressed( "SelectAndMove_FullyPressed" );
 	static const FName Undo( "Undo" );
 	static const FName Redo( "Redo" );
 	static const FName Delete( "Delete" );
@@ -73,6 +75,10 @@ public:
 	DECLARE_EVENT_ThreeParams( UViewportWorldInteraction, FOnVIHoverUpdate, class UViewportInteractor* /* Interactor */, FVector& /* OutHoverImpactPoint */, bool& /* bWasHandled */ );
 	FOnVIHoverUpdate& OnViewportInteractionHoverUpdate() { return OnHoverUpdateEvent; }
 
+	/** Gets the event for previewing input actions.  This allows systems to get a first chance at actions before everything else. */
+	DECLARE_EVENT_FiveParams( UViewportWorldInteraction, FOnPreviewInputAction, class FEditorViewportClient& /* ViewportClient */, class UViewportInteractor* /* Interactor */, const struct FViewportActionKeyInput& /* Action */, bool& /* bOutIsInputCaptured */, bool& /* bWasHandled */ );
+	FOnPreviewInputAction& OnPreviewInputAction() { return OnPreviewInputActionEvent; }
+
 	/** Gets the event for input actions update which is broadcasted for each interactor */
 	DECLARE_EVENT_FiveParams( UViewportWorldInteraction, FOnVIActionHandle, class FEditorViewportClient& /* ViewportClient */, class UViewportInteractor* /* Interactor */, const struct FViewportActionKeyInput& /* Action */, bool& /* bOutIsInputCaptured */, bool& /* bWasHandled */ );
 	FOnVIActionHandle& OnViewportInteractionInputAction() { return OnInputActionEvent; }
@@ -81,11 +87,11 @@ public:
 	DECLARE_EVENT_ThreeParams( UViewportWorldInteraction, FOnViewportInteractionInputUnhandled, class FEditorViewportClient& /* ViewportClient */, class UViewportInteractor* /* Interactor */, const struct FViewportActionKeyInput& /* Action */ );
 	FOnViewportInteractionInputUnhandled& OnViewportInteractionInputUnhandled() { return OnViewportInteractionInputUnhandledEvent; };
 
-	/** To handle raw key input from the Inputprocessor */
+	/** To handle raw key input from the Input Preprocessor */
 	DECLARE_EVENT_FourParams( UViewportWorldInteraction, FOnHandleInputKey, class FEditorViewportClient* /* ViewportClient */, const FKey /* Key */, const EInputEvent /* Event */, bool& /* bWasHandled */ );
 	FOnHandleInputKey& OnHandleKeyInput() { return OnKeyInputEvent; }
 
-	/** To handle raw axis input from the Inputprocessor */
+	/** To handle raw axis input from the Input Preprocessor */
 	DECLARE_EVENT_SixParams( UViewportWorldInteraction, FOnHandleInputAxis, class FEditorViewportClient* /* ViewportClient */, const int32 /* ControllerId */, const FKey /* Key */, const float /* Delta */, const float /* DeltaTime */, bool& /* bWasHandled */ );
 	FOnHandleInputAxis& OnHandleAxisInput() { return OnAxisInputEvent; }
 
@@ -109,6 +115,12 @@ public:
 	    nullptr, a default transformer for actors/components will be enabled */
 	void SetTransformer( class UViewportTransformer* NewTransformer );
 
+	/** @return	Gets the currently active transformer object */
+	const UViewportTransformer* GetTransformer() const
+	{
+		return ViewportTransformer;
+	}
+
 	/** Sets the list of objects that this system will be responsible for transforming when interacting using the
 	    gizmo or directly on the objects */
 	void SetTransformables( TArray< TUniquePtr< FViewportTransformable > >&& NewTransformables );
@@ -119,6 +131,20 @@ public:
 	/** Pairs to interactors by setting	the other interactor for each interactor */
 	void PairInteractors( UViewportInteractor* FirstInteractor, UViewportInteractor* SecondInteractor );
 	
+	/**
+	 * Adds an actor to the list of actors to never allow an interactor to hit in the scene.  No selection.  No hover.
+	 * There's no need to remove actors from this list.  They'll expire from it automatically when destroyed.
+	 *
+	 * @param	ActorToExcludeFromHitTests	The actor that should be forever excluded from hit tests
+	 */
+	void AddActorToExcludeFromHitTests( AActor* ActorToExcludeFromHitTests );
+
+	/** @return	Returns the list of actors that should never be included in interctor hit tests (hovering or selection) */
+	const TArray<TWeakObjectPtr<AActor>>& GetActorsToExcludeFromHitTest() const
+	{
+		return ActorsToExcludeFromHitTest;
+	}
+
 
 	//
 	// Input
@@ -127,10 +153,10 @@ public:
 	virtual bool InputKey( FEditorViewportClient* InViewportClient, FViewport* Viewport, FKey Key, EInputEvent Event ) override;
 	virtual bool InputAxis( FEditorViewportClient* InViewportClient, FViewport* Viewport, int32 ControllerId, FKey Key, float Delta, float DeltaTime ) override;
 
-	/** Handles the key input from the InputProcessor*/
+	/** Handles the key input from the Input Preprocessor*/
 	bool PreprocessedInputKey( const FKey Key, const EInputEvent Event );
 
-	/** Handles the axis input from the InputProcessor*/
+	/** Handles the axis input from the Input Preprocessor*/
 	bool PreprocessedInputAxis( const int32 ControllerId, const FKey Key, const float Delta, const float DeltaTime );
 
 	/** Gets the world space transform of the calibrated VR room origin.  When using a seated VR device, this will feel like the
@@ -152,9 +178,6 @@ public:
 	/** Gets the world scale factor, which can be multiplied by a scale vector to convert to room space */
 	float GetWorldScaleFactor() const;
 
-	/** Gets the world of the viewport used for this worldinteraction */
-	virtual UWorld* GetWorld() const;
-
 	/** Gets the currently used viewport */
 	FEditorViewportClient* GetDefaultOptionalViewportClient() const;
 
@@ -188,8 +211,8 @@ public:
 	/** Called to finish a drag action with the specified interactor */
 	void StopDragging( class UViewportInteractor* Interactor );
 
-	/** Starts dragging selected objects around.  Called when clicking and dragging on actors/gizmos in the world, or when placing new objects.  Returns true if dragging actually started. */
-	bool StartDragging( UViewportInteractor* Interactor, UActorComponent* ClickedTransformGizmoComponent, const FVector& HitLocation, const bool bIsPlacingNewObjects, const bool bStartTransaction );
+	/** Starts dragging selected objects around.  Called when clicking and dragging on actors/gizmos in the world, or when placing new objects. */
+	void StartDragging( UViewportInteractor* Interactor, UActorComponent* ClickedTransformGizmoComponent, const FVector& HitLocation, const bool bIsPlacingNewObjects, const bool bAllowInterpolationWhenPlacing, const bool bStartTransaction, const bool bWithGrabberSphere );
 
 	DECLARE_EVENT_OneParam( UViewportWorldInteraction, FOnWorldScaleChanged, const float /* NewWorldToMetersScale */);
 	virtual FOnWorldScaleChanged& OnWorldScaleChanged() { return OnWorldScaleChangedEvent; };
@@ -207,7 +230,7 @@ public:
 	float GetMinScale();
 
 	/** Sets GNewWorldToMetersScale */
-	void SetWorldToMetersScale( const float NewWorldToMetersScale );
+	void SetWorldToMetersScale( const float NewWorldToMetersScale, const bool bCompensateRoomWorldScale = false );
 
 	/** Tells the world interaction system to skip updates of world movement this frame.  This is useful if you've called
 	    SetWorldToMetersScale() yourself to set a static world scale, and the cached room-space transforms for the last
@@ -261,7 +284,7 @@ public:
 	void SetTransformGizmoClass( const TSubclassOf<ABaseTransformGizmo>& NewTransformGizmoClass );
 
 	/** Sets the currently dragged interactavle */
-	void SetDraggedInteractable( IViewportInteractableInterface* InDraggedInteractable );	
+	void SetDraggedInteractable( IViewportInteractableInterface* InDraggedInteractable, UViewportInteractor* Interactor );	
 
 	/** Check if there is another interactor hovering over the component */
 	bool IsOtherInteractorHoveringOverComponent( UViewportInteractor* Interactor, UActorComponent* Component ) const;
@@ -275,10 +298,8 @@ public:
 	/** Gets the tracking transactions */
 	FTrackingTransaction& GetTrackingTransaction();
 	
-	void SetUseInputProcessor( bool bInUseInputProcessor );
-
-	/** Resets the start transform all of transformables to match current transform */
-	void ResetTransformables();
+	/** Sets whether we should allow input events from the input preprocessor or not */
+	void SetUseInputPreprocessor( bool bInUseInputPreprocessor );
 
 	/** Gets a list of all of the objects we're currently interacting with, such as selected actors */
 	TArray< TUniquePtr< FViewportTransformable > >& GetTransformables()
@@ -286,22 +307,49 @@ public:
 		return Transformables;
 	}
 
-	/** Static: Spawns a transient actor that we can use in the editor world (templated for convenience) */
-	template<class T>
-	static inline T* SpawnTransientSceneActor(UWorld* World, const FString& ActorName, const bool bWithSceneComponent)
+	/** The ability to move and scale the world */
+	void AllowWorldMovement(bool bDisable);
+
+	/** For other systems to check if the Viewport World Interaction system is currently aligning transformables to actors*/
+	bool AreAligningToActors();
+	/** For other systems to check if the Viewport World Interaction system currently has candidate actors selected */
+	bool HasCandidatesSelected();
+	/** If there are no currently selected candidates, use the currently selected actors as candidates. Otherwise, reset the candidates */
+	void SetSelectionAsCandidates();
+
+	/** When RotateOnAngle we intersect on a plane to rotate the transform gizmo. This is the local point from the transform gizmo location of that intersect */
+	FVector GetLocalIntersectPointOnRotationGizmo() const;
+
+	/** Gets the current delta time, so functions that don't get the delta time passed can still get it */
+	float GetCurrentDeltaTime() const;
+
+	/** Getters and setters for whether or not to show the cursor on the viewport */
+	bool ShouldSuppressExistingCursor() const;
+	void SetShouldSuppressExistingCursor(const bool bInShouldSuppressCursor)
 	{
-		return CastChecked<T>(SpawnTransientSceneActor(World, T::StaticClass(), ActorName, bWithSceneComponent));
-	}
+		bShouldSuppressCursor = bInShouldSuppressCursor;
+	};
 
-	/** Static: Spawns a transient actor that we can use in the editor world */
-	static AActor* SpawnTransientSceneActor( UWorld* World, TSubclassOf<AActor> ActorClass, const FString& ActorName, const bool bWithSceneComponent );
+	/** Gets the container for all the assets of ViewportInteraction. */
+	const class UViewportInteractionAssetContainer& GetAssetContainer() const;
 
-	/** Static: Destroys a transient actor we created earlier */
-	static void DestroyTransientActor( UWorld* World, AActor* Actor );
+	/** Static function to load the asset container */
+	static const class UViewportInteractionAssetContainer& LoadAssetContainer();
 
+	/** Plays sound at location. */
+	void PlaySound(USoundBase* SoundBase, const FVector& InWorldLocation, const float InVolume = 1.0f);
+
+	/** Set if this world interaction is in VR. */
+	void SetInVR(const bool bInVR);
+
+	/** Get if this world interaction is in VR. */
+	bool IsInVR() const;
 
 protected:
 
+	virtual void TransitionWorld(UWorld* NewWorld) override;
+	virtual void EnteredSimulateInEditor() override;
+	virtual void LeftSimulateInEditor() override;
 
 private:
 
@@ -334,13 +382,20 @@ private:
 		const float LaserPointerMaxLength,
 		const bool bIsLaserPointerValid, 
 		const FTransform& GizmoStartTransform, 
-		const FBox& GizmoStartLocalBounds, 
+		FTransform& GizmoLastTransform,
+		FTransform& GizmoTargetTransform,
+		FTransform& GizmoUnsnappedTargetTransform,
+		const FTransform& GizmoInterpolationSnapshotTransform,
+		const FBox& GizmoStartLocalBounds,
 		const USceneComponent* const DraggingTransformGizmoComponent,
 		FVector& GizmoSpaceFirstDragUpdateOffsetAlongAxis,
 		FVector& DragDeltaFromStartOffset,
 		bool& bIsDrivingVelocityOfSimulatedTransformables,
 		FVector& OutUnsnappedDraggedTo);
-	
+
+	FVector SnapLocation(const bool bLocalSpaceSnapping, const FVector& DesiredGizmoLocation, const FTransform &GizmoStartTransform, const FVector SnapGridBase, const bool bShouldConstrainMovement, const FVector AlignAxes);
+
+
 	/** Given a drag delta from a starting point, contrains that delta based on a gizmo handle axis */
 	FVector ComputeConstrainedDragDeltaFromStart( 
 		const bool bIsFirstDragUpdate, 
@@ -373,7 +428,7 @@ private:
 	void PollInputIfNeeded();
 
 	/** Refreshes the transform gizmo to match the currently selected actor set */
-	void RefreshTransformGizmo( const bool bNewObjectsSelected, bool bAllHandlesVisible, class UActorComponent* SingleVisibleHandle, const TArray< UActorComponent* >& HoveringOverHandles );
+	void RefreshTransformGizmo( const bool bNewObjectsSelected );
 
 	/** Spawn new transform gizmo when refreshing and at start if there is none yet or the current one is different from the TransformGizmoClass */
 	void SpawnTransformGizmoIfNeeded();
@@ -384,15 +439,6 @@ private:
 	/** Destroys the actors */
 	void DestroyActors();
 
-public:
-	/** Gets the snap grid mesh actor */
-	AActor* GetSnapGridActor()
-	{
-		SpawnGridMeshActor();
-		return SnapGridActor;
-	}
-
-private:
 	/** Gets the snap grid mesh MID */
 	class UMaterialInstanceDynamic* GetSnapGridMID()
 	{
@@ -406,6 +452,16 @@ private:
 	/** Average location of all the current transformables */
 	FVector CalculateAverageLocationOfTransformables();
 
+	/** Location the transformable should snap to if aligning to other transformables*/
+	FVector FindTransformGizmoAlignPoint(const FTransform& GizmoStartTransform, const FTransform& DesiredGizmoTransform, const bool bShouldConstrainMovement, FVector ConstraintAxes);
+	void DrawBoxBrackets(const FBox InActor, const FTransform LocalToWorld, const FLinearColor BracketColor);
+	
+	/** Calculate a new room transform location according to the world to meters scale */
+	void CompensateRoomTransformForWorldScale(FTransform& InOutRoomTransform, const float InNewWorldToMetersScale, const FVector& InRoomPivotLocation);
+
+	/** If there is a transformable with velocity in simulate */
+	bool HasTransformableWithVelocityInSimulate() const;
+
 	//
 	// Colors
 	//
@@ -416,19 +472,25 @@ public:
 		Forward,
 		Right,
 		Up,
-		Hover,
-		Dragging,
+		GizmoHover,
+		GizmoDragging,
 		TotalCount
 	};
 
 	/** Gets the color from color type */
-	FLinearColor GetColor( const EColors Color ) const;
+	FLinearColor GetColor(const EColors Color, const float Multiplier = 1.f) const;
+
+	/** The path of the asset container */
+	static const FString AssetContainerPath;
 
 private:
 
 	// All the colors for this mode
 	TArray<FLinearColor> Colors;
 
+	/** Candidate actors for aligning in the scene */
+	TArray<const AActor*> CandidateActors;
+
 protected:
 
 	/** True if we've dragged objects with either hand since the last time we selected something */
@@ -482,6 +544,12 @@ private:
 	/** True if we should skip interactive world movement/rotation/scaling this frame (because it was already set by something else) */
 	bool bSkipInteractiveWorldMovementThisFrame;
 
+	/** The room transform to apply at the beginning of the next frame (if needed.)  This is used to defer
+	    updates of the room transform to match with the GNewWorldToMetersScale, which will also take effect
+		on the next frame.  It's critical that those two are applied at the same time, because otherwise a
+		frame rate dependent feedback loop will occur */
+	TOptional<TTuple<FTransform,uint32>> RoomTransformToSetOnFrame;
+
 
 	//
 	// Hover state
@@ -501,6 +569,9 @@ private:
 	/** True if our transformables are currently interpolating to their target location */
 	bool bIsInterpolatingTransformablesFromSnapshotTransform;
 
+	/** When interpolating from a snapshot transform, this will be true if we should freeze the destination placement location until the interpolation finishes */
+	bool bFreezePlacementWhileInterpolatingTransformables;
+
 	/** Time that we started interpolating at */
 	FTimespan TransformablesInterpolationStartTime;
 
@@ -524,8 +595,10 @@ private:
 	/** Starting angle when rotating an object with  ETransformGizmoInteractionType::RotateOnAngle */
 	TOptional<float> StartDragAngleOnRotation;
 
+	FVector LocalIntersectPointOnRotationGizmo;
+
 	/** The direction of where the rotation handle is facing when starting rotation */
-	TOptional<FVector> StartDragHandleDirection;
+	TOptional<FVector> DraggingRotationHandleDirection;
 
 	/** Whether the gizmo should be visible or not */
 	bool bShouldTransformGizmoBeVisible;
@@ -566,6 +639,9 @@ private:
 	/** Event that fires every frame to update hover state based on what's under the cursor.  Set bWasHandled to true if you detect something to hover. */
 	FOnVIHoverUpdate OnHoverUpdateEvent;
 
+	/** Event for previewing input actions.  This allows systems to get a first chance at actions before everything else. */
+	FOnPreviewInputAction OnPreviewInputActionEvent;
+
 	/** Event that is fired for when a key is pressed by an interactor */
 	FOnVIActionHandle OnInputActionEvent;
 
@@ -605,6 +681,16 @@ private:
 	/** Reference count for the default mouse cursor interactor.  When this reaches zero, the mouse cursor interactor goes away. */
 	int32 DefaultMouseCursorInteractorRefCount;
 
+	/** List of actors which should never be hit by an interactor, such as the 'avatar mesh actor' in VR */
+	UPROPERTY()
+	TArray<TWeakObjectPtr<AActor>> ActorsToExcludeFromHitTest;
+
+	//
+	// VR
+	//
+	
+	/** If this world interaction is in VR */
+	bool bIsInVR;
 
 	//
 	// Configuration
@@ -617,5 +703,24 @@ private:
 	bool bActive;
 
 	/** If this world interaction should get input from the input processor */
-	bool bUseInputProcessor;
+	bool bUseInputPreprocessor;
+
+	/** If the user can scale and navigate around in the world */
+	bool bAllowWorldMovement;
+
+	/** The delta time of this tick */
+	float CurrentDeltaTime;
+
+	/** Whether or not to show the cursor on the viewport */
+	bool bShouldSuppressCursor;
+
+	/** The current tick number */
+	uint32 CurrentTickNumber;
+
+	/** Container of assets */
+	UPROPERTY()
+	class UViewportInteractionAssetContainer* AssetContainer;
+
+	/** If we want to skip playing the sound when refreshing the transform gizmo next time */
+	bool bPlayNextRefreshTransformGizmoSound;
 };
diff --git a/Engine/Source/Editor/ViewportInteraction/ViewportDragOperation.cpp b/Engine/Source/Editor/ViewportInteraction/ViewportDragOperation.cpp
--- a/Engine/Source/Editor/ViewportInteraction/ViewportDragOperation.cpp
+++ b/Engine/Source/Editor/ViewportInteraction/ViewportDragOperation.cpp
@@ -44,3 +44,8 @@ void UViewportDragOperationComponent::ClearDragOperation()
 
 	DragOperation = nullptr;
 }
+
+bool UViewportDragOperationComponent::IsDragging() const
+{
+	return DragOperation != nullptr;
+}
diff --git a/Engine/Source/Editor/ViewportInteraction/ViewportInteractor.cpp b/Engine/Source/Editor/ViewportInteraction/ViewportInteractor.cpp
--- a/Engine/Source/Editor/ViewportInteraction/ViewportInteractor.cpp
+++ b/Engine/Source/Editor/ViewportInteraction/ViewportInteractor.cpp
@@ -7,14 +7,25 @@
 #include "ViewportWorldInteraction.h"
 #include "ViewportInteractableInterface.h"
 #include "MouseCursorInteractor.h"
+#include "ScopedTransaction.h"
+#include "DrawDebugHelpers.h"
+#include "EngineUtils.h"
+#include "ActorTransformer.h"
+#include "Engine/Selection.h"
+
+#define LOCTEXT_NAMESPACE "ViewportInteractor"
 
 namespace VI
 {
-	static FAutoConsoleVariable GrabberSphereRadius( TEXT( "VI.GrabberSphereRadius" ), 15.0f, TEXT( "In radial mode, the radius of the sphere used to select and interact" ) );
-	static FAutoConsoleVariable GrabberSphereOffset( TEXT( "VI.GrabberSphereOffset" ), 10.0f, TEXT( "Offset from the controller origin that the grabber sphere should be centered at" ) );
-	static FAutoConsoleVariable LaserPointerMaxLength( TEXT( "VI.LaserPointerMaxLength" ), 10000.0f, TEXT( "Maximum length of the laser pointer line" ) );
+	static FAutoConsoleVariable GrabberSphereRadius( TEXT( "VI.GrabberSphereRadius" ), 12.0f, TEXT( "In radial mode, the radius of the sphere used to select and interact" ) );
+	static FAutoConsoleVariable GrabberSphereOffset( TEXT( "VI.GrabberSphereOffset" ), 2.0f, TEXT( "Offset from the controller origin that the grabber sphere should be centered at" ) );
+	static FAutoConsoleVariable LaserPointerMaxLength( TEXT( "VI.LaserPointerMaxLength" ), 30000.0f, TEXT( "Maximum length of the laser pointer line" ) );
 	static FAutoConsoleVariable DragHapticFeedbackStrength( TEXT( "VI.DragHapticFeedbackStrength" ), 1.0f, TEXT( "Default strength for haptic feedback when starting to drag objects" ) ); //@todo ViewportInteraction: Duplicate from ViewportWorldInteraction
 	static FAutoConsoleVariable SelectionHapticFeedbackStrength( TEXT( "VI.SelectionHapticFeedbackStrength" ), 0.5f, TEXT( "Default strength for haptic feedback when selecting objects" ) );
+	static FAutoConsoleVariable LaserPointerRetractDuration( TEXT( "VI.LaserPointerRetractDuration" ), 0.2f, TEXT( "How fast the laser pointer should extend or retract" ) );
+	static FAutoConsoleVariable AllowLaserSmooth(TEXT("VI.AllowLaserSmooth"), 1, TEXT("Allow laser smoothing using one euro"));
+	static FAutoConsoleVariable LaserSmoothLag(TEXT("VI.LaserSmoothLag"), 0.007f, TEXT("Laser smooth lag"));
+	static FAutoConsoleVariable LaserSmoothMinimumCutoff(TEXT("VI.LaserSmoothMinimumCutoff"), 0.9f, TEXT("Laser smooth lag"));
 }
 
 UViewportInteractor::UViewportInteractor( const FObjectInitializer& Initializer ) : 
@@ -23,8 +34,10 @@ UViewportInteractor::UViewportInteractor( const FObjectInitializer& Initializer
 	KeyToActionMap(),
 	WorldInteraction( nullptr ),
 	OtherInteractor( nullptr ),
-	bAllowGrabberSphere( true )
+	bAllowGrabberSphere( true ),
+	SavedLaserPointerEnd()
 {
+	SmoothingOneEuroFilter = ViewportInteractionUtils::FOneEuroFilter(VI::LaserSmoothMinimumCutoff->GetFloat(), VI::LaserSmoothLag->GetFloat(), 1.0f);
 }
 
 FViewportInteractorData& UViewportInteractor::GetInteractorData()
@@ -52,14 +65,6 @@ UViewportInteractor* UViewportInteractor::GetOtherInteractor() const
 	return OtherInteractor;
 }
 
-void UViewportInteractor::UpdateHoverResult( FHitResult& OutHitResult )
-{
-	FHitResult HitResult = GetHitResultFromLaserPointer();
-	InteractorData.bIsHovering = HitResult.bBlockingHit;
-	InteractorData.HoverLocation = HitResult.ImpactPoint;
-	OutHitResult = HitResult;
-}
-
 void UViewportInteractor::Shutdown()
 {
 	KeyToActionMap.Empty();
@@ -68,9 +73,27 @@ void UViewportInteractor::Shutdown()
 	OtherInteractor = nullptr;
 }
 
-class UActorComponent* UViewportInteractor::GetHoverComponent()
+void UViewportInteractor::Tick( const float DeltaTime )
+{
+}
+
+void UViewportInteractor::OnStartDragging( const FVector& HitLocation, const bool bIsPlacingNewObjects )
+{
+	// If the user is holding down the modifier key, go ahead and duplicate the selection first.  Don't do this if we're
+	// placing objects right now though.
+	if( IsModifierPressed() && !bIsPlacingNewObjects )
+	{
+		// Only duplicate selected objects if we're using the "actor" transformer.  Other transformable types might not support duplication!
+		if( Cast<const UActorTransformer>( WorldInteraction->GetTransformer() ) != nullptr )
+		{
+			WorldInteraction->Duplicate();
+		}
+	}
+}
+
+class UActorComponent* UViewportInteractor::GetLastHoverComponent()
 {
-	return InteractorData.HoveredActorComponent.Get();
+	return InteractorData.LastHoveredActorComponent.Get();
 }
 
 void UViewportInteractor::AddKeyAction( const FKey& Key, const FViewportActionKeyInput& Action )
@@ -91,18 +114,33 @@ bool UViewportInteractor::HandleInputKey( FEditorViewportClient& ViewportClient,
 	{
 		Action->Event = Event;
 		
-		// Prefer transform gizmo interactions over anything else
+		if( !bHandled )
+		{
+			// Give the derived classes a chance to update according to the input
+			PreviewInputKey( ViewportClient, *Action, Key, Event, bHandled );
+		}
+
 		FHitResult HitResult = GetHitResultFromLaserPointer();
-		const bool bPressedTransformGizmo =
-			Event == IE_Pressed &&
-			HitResult.Actor.IsValid() &&
-			HitResult.Actor == WorldInteraction->GetTransformGizmoActor() &&
-			( Action->ActionType == ViewportWorldActionTypes::SelectAndMove ||
-			  Action->ActionType == ViewportWorldActionTypes::SelectAndMove_LightlyPressed );
-		if( !bPressedTransformGizmo )
+
+		if( !bHandled )
+		{
+			// "Preview" the event first.  This gives a first chance for systems to intercept the event regardless of registration order.
+			WorldInteraction->OnPreviewInputAction().Broadcast( ViewportClient, this, *Action, Action->bIsInputCaptured, bHandled );
+		}
+
+		if( !bHandled )
 		{
-			// Give subsystems a chance to handle actions for this interactor
-			WorldInteraction->OnViewportInteractionInputAction().Broadcast( ViewportClient, this, *Action, Action->bIsInputCaptured, bHandled );
+			// Prefer transform gizmo interactions over anything else
+			const bool bPressedTransformGizmo =
+				Event == IE_Pressed &&
+				HitResult.Actor.IsValid() &&
+				HitResult.Actor == WorldInteraction->GetTransformGizmoActor() &&
+				Action->ActionType == ViewportWorldActionTypes::SelectAndMove;
+			if( !bPressedTransformGizmo )
+			{
+				// Give subsystems a chance to handle actions for this interactor
+				WorldInteraction->OnViewportInteractionInputAction().Broadcast( ViewportClient, this, *Action, Action->bIsInputCaptured, bHandled );
+			}
 		}
 
 		if(!bHandled)
@@ -115,8 +153,7 @@ bool UViewportInteractor::HandleInputKey( FEditorViewportClient& ViewportClient,
 		if ( !bHandled )
 		{
 			// Selection/Movement
-			if ( ( Action->ActionType == ViewportWorldActionTypes::SelectAndMove ||
-				   Action->ActionType == ViewportWorldActionTypes::SelectAndMove_LightlyPressed ) )
+			if( Action->ActionType == ViewportWorldActionTypes::SelectAndMove )
 			{
 				const bool bIsDraggingWorld = InteractorData.DraggingMode == EViewportInteractionDraggingMode::World;
 				const bool bIsDraggingWorldWithTwoHands =
@@ -127,7 +164,7 @@ bool UViewportInteractor::HandleInputKey( FEditorViewportClient& ViewportClient,
 				if ( Event == IE_Pressed )
 				{
 					// No clicking while we're dragging the world.  (No laser pointers are visible, anyway.)
-					if ( !bIsDraggingWorldWithTwoHands && !bHandled && HitResult.Actor.IsValid() )
+					if ( !bIsDraggingWorldWithTwoHands && HitResult.Actor.IsValid() )
 					{
 						if ( WorldInteraction->IsInteractableComponent( HitResult.GetComponent() ) )
 						{
@@ -141,18 +178,13 @@ bool UViewportInteractor::HandleInputKey( FEditorViewportClient& ViewportClient,
 
 							if ( IViewportInteractableInterface* ActorInteractable = Cast<IViewportInteractableInterface>( Actor ) )
 							{
-								if ( Action->ActionType == ViewportWorldActionTypes::SelectAndMove_LightlyPressed )
-								{
-									bHandled = true;
-									SetAllowTriggerFullPress( false );
-									bool bResultedInInteractableDrag = false;
-									ActorInteractable->OnPressed( this, HitResult, bResultedInInteractableDrag );
+								bHandled = true;
+								bool bResultedInInteractableDrag = false;
+								ActorInteractable->OnPressed( this, HitResult, bResultedInInteractableDrag );
 
-									if ( bResultedInInteractableDrag )
-									{
-										InteractorData.DraggingMode = EViewportInteractionDraggingMode::Interactable;
-										WorldInteraction->SetDraggedInteractable( ActorInteractable );
-									}
+								if ( bResultedInInteractableDrag )
+								{
+									WorldInteraction->SetDraggedInteractable(ActorInteractable, this);
 								}
 							}
 							else
@@ -180,21 +212,22 @@ bool UViewportInteractor::HandleInputKey( FEditorViewportClient& ViewportClient,
 
 											UPrimitiveComponent* ClickedTransformGizmoComponent = nullptr;
 											const bool bIsPlacingNewObjects = false;
-											const bool bStartTransaction = true;
-											if( WorldInteraction->StartDragging( OtherInteractor, ClickedTransformGizmoComponent, OtherInteractorData->HoverLocation, bIsPlacingNewObjects, bStartTransaction ) )
-											{
-												// ...
-											}
+											const bool bAllowInterpolationWhenPlacing = false;
+											const bool bStartTransaction = !WorldInteraction->GetTrackingTransaction().IsActive();
+											const bool bWithGrabberSphere = false;	// @todo grabber: Not supported yet
+											WorldInteraction->StartDragging( OtherInteractor, ClickedTransformGizmoComponent, OtherInteractor->GetHoverLocation(), bIsPlacingNewObjects, bAllowInterpolationWhenPlacing, bStartTransaction, bWithGrabberSphere );
 										}
 
 										InteractorData.DraggingMode = InteractorData.LastDraggingMode = EViewportInteractionDraggingMode::AssistingDrag;
 
+										InteractorData.bDraggingWithGrabberSphere = false;	// @todo grabber: Not supported yet
 										InteractorData.bIsFirstDragUpdate = true;
 										InteractorData.bWasAssistingDrag = true;
 										InteractorData.DragRayLength = ( HitResult.ImpactPoint - LaserPointerStart ).Size();
-										InteractorData.LastLaserPointerStart = LaserPointerStart;
 										InteractorData.LastDragToLocation = HitResult.ImpactPoint;
-										InteractorData.LaserPointerImpactAtDragStart = HitResult.ImpactPoint;
+										InteractorData.InteractorRotationAtDragStart = InteractorData.Transform.GetRotation();
+										InteractorData.GrabberSphereLocationAtDragStart = FVector::ZeroVector;
+										InteractorData.ImpactLocationAtDragStart = HitResult.ImpactPoint;
 										InteractorData.DragTranslationVelocity = FVector::ZeroVector;
 										InteractorData.DragRayLengthVelocity = 0.0f;
 										InteractorData.DraggingTransformGizmoComponent = nullptr;
@@ -202,6 +235,7 @@ bool UViewportInteractor::HandleInputKey( FEditorViewportClient& ViewportClient,
 										InteractorData.bIsDrivingVelocityOfSimulatedTransformables = false;
 
 										InteractorData.GizmoStartTransform = OtherInteractorData->GizmoStartTransform;
+										InteractorData.GizmoLastTransform = InteractorData.GizmoTargetTransform = InteractorData.GizmoUnsnappedTargetTransform = InteractorData.GizmoInterpolationSnapshotTransform = InteractorData.GizmoStartTransform;
 										InteractorData.GizmoStartLocalBounds = OtherInteractorData->GizmoStartLocalBounds;
 										InteractorData.GizmoSpaceFirstDragUpdateOffsetAlongAxis = FVector::ZeroVector;	// Will be determined on first update
 										InteractorData.GizmoSpaceDragDeltaFromStartOffset = FVector::ZeroVector;	// Set every frame while dragging
@@ -229,87 +263,78 @@ bool UViewportInteractor::HandleInputKey( FEditorViewportClient& ViewportClient,
 								}
 								else
 								{
-									// Default to replacing our selection with whatever the user clicked on
-									enum ESelectionModification
-									{
-										Replace,
-										AddTo,
-										Toggle
-									} SelectionModification = ESelectionModification::Replace;
-
-									// If the other hand is holding down the button after selecting an object, we'll allow this hand to toggle selection
-									// of additional objects (multi select)
-									{
-										if( OtherInteractorData != nullptr && 
-											OtherInteractorData->ClickingOnComponent.IsValid() &&		// Other hand is clicking on something
-											Actor != WorldInteraction->GetTransformGizmoActor() &&					// Not clicking on a gizmo
-											OtherInteractorData->ClickingOnComponent.Get()->GetOwner() != Actor )	// Not clicking on same actor
-										{
-											// OK, the other hand is holding down the "select and move" button.
-											SelectionModification = ESelectionModification::Toggle;
-										}
-									}
-
 									// Only start dragging if the actor was already selected (and it's a full press), or if we clicked on a gizmo.  It's too easy to 
 									// accidentally move actors around the scene when you only want to select them otherwise.
-									bool bShouldDragSelected =
-										( ( Actor->IsSelected() && Action->ActionType == ViewportWorldActionTypes::SelectAndMove ) || Actor == WorldInteraction->GetTransformGizmoActor() ) &&
-										( SelectionModification == ESelectionModification::Replace );
+									const bool bOtherHandTryingToDrag =
+										Actor != WorldInteraction->GetTransformGizmoActor() &&
+										OtherInteractorData != nullptr &&
+										OtherInteractorData->ClickingOnComponent.IsValid() &&
+										OtherInteractorData->ClickingOnComponent.Get()->GetOwner()->IsSelected() &&
+										OtherInteractorData->ClickingOnComponent.Get()->GetOwner() == HitResult.GetComponent()->GetOwner();	// Trying to drag same actor
+
+									FTrackingTransaction& TrackingTransaction = WorldInteraction->GetTrackingTransaction();
 
+									bool bShouldDragSelected = false;
 									bool bSelectionChanged = false;
-									if ( Actor != WorldInteraction->GetTransformGizmoActor() )  // Don't change any actor selection state if the user clicked on a gizmo
+									if( Actor == WorldInteraction->GetTransformGizmoActor() )
 									{
+										bShouldDragSelected = true;
+										bSelectionChanged = false;
+									}
+									else
+									{
+										// Clicked on a normal actor.  So update selection!
 										const bool bWasSelected = Actor->IsSelected();
 
-										// Clicked on a normal actor.  So update selection!
-										if ( SelectionModification == ESelectionModification::Replace || SelectionModification == ESelectionModification::AddTo )
+										// Determine whether we are starting a marquee select
+										const bool bIsMarqueeSelect = ViewportClient.IsAltPressed() && ViewportClient.IsCtrlPressed();
+
+										// Default to replacing our selection with whatever the user clicked on
+										enum ESelectionModification
+										{
+											Replace,
+											Toggle
+										} SelectionModification = ESelectionModification::Replace;
+
+										// Light pressing on an actor while holding the modifier button will toggle selection
+										const bool bIsMouseCursorInteractor = this->IsA( UMouseCursorInteractor::StaticClass() );
+										if( this->IsModifierPressed() && !bWasSelected )
 										{
-											if ( !bWasSelected )
+											SelectionModification = ESelectionModification::Toggle;
+										}
+
+										if( !bIsMarqueeSelect )
+										{
+											if( !bWasSelected && SelectionModification == ESelectionModification::Replace )
 											{
-												if ( SelectionModification == ESelectionModification::Replace )
-												{
-													GEditor->SelectNone( true, true );
-												}
+												// Capture undo state
+												bSelectionChanged = true;
+												TrackingTransaction.TransCount++;
+												TrackingTransaction.Begin( LOCTEXT( "UndoSelectingActors", "Select Actor" ) );
 
+												GEditor->SelectNone( true, true );
 												GEditor->SelectActor( Actor, true, true );
+											}
+											else if ( SelectionModification == ESelectionModification::Toggle )
+											{
+												TrackingTransaction.TransCount++;
+												TrackingTransaction.Begin( LOCTEXT( "UndoSelectingActors", "Toggle Actor Selection" ) );
+
+												GEditor->SelectActor( Actor, !Actor->IsSelected(), true );
 												bSelectionChanged = true;
 											}
 										}
-										else if ( SelectionModification == ESelectionModification::Toggle )
-										{
-											GEditor->SelectActor( Actor, !Actor->IsSelected(), true );
-											bSelectionChanged = true;
-										}
 
-										// If we've selected an actor with a light press, don't prevent the user from moving the actor if they
-										// press down all the way.  We only want locking when interacting with gizmos
-										SetAllowTriggerLightPressLocking( false );
-
-										// If the user did a full press to select an actor that wasn't selected, allow them to drag it right away
-										const bool bOtherHandTryingToDrag =
-											OtherInteractorData != nullptr &&
-											OtherInteractorData->ClickingOnComponent.IsValid() &&
-											OtherInteractorData->ClickingOnComponent.Get()->GetOwner()->IsSelected() &&
-											OtherInteractorData->ClickingOnComponent.Get()->GetOwner() == HitResult.GetComponent()->GetOwner();	// Trying to drag same actor
-										if ( ( !bWasSelected || bOtherHandTryingToDrag ) &&
-											Actor->IsSelected() &&
-											( Action->ActionType == ViewportWorldActionTypes::SelectAndMove || bOtherHandTryingToDrag ) )
-										{
-											bShouldDragSelected = true;
-										}
+										bShouldDragSelected = bOtherHandTryingToDrag || !bSelectionChanged;
 									}
 
 									if ( bShouldDragSelected && InteractorData.DraggingMode != EViewportInteractionDraggingMode::Interactable )
 									{
 										UPrimitiveComponent* ClickedTransformGizmoComponent = nullptr;
 										{
-											const bool bIsMouseCursorInteractor = this->IsA( UMouseCursorInteractor::StaticClass() );
-
 											const bool bUsingGizmo =
 												HitResult.GetComponent() != nullptr &&
-												HitResult.GetComponent()->GetOwner() == WorldInteraction->GetTransformGizmoActor() &&
-												( Action->ActionType == ViewportWorldActionTypes::SelectAndMove_LightlyPressed ||
-												( bIsMouseCursorInteractor && Action->ActionType == ViewportWorldActionTypes::SelectAndMove ) );		// Only use the gizmo when lightly pressed (unless we're dealing with the mouse cursor)
+												HitResult.GetComponent()->GetOwner() == WorldInteraction->GetTransformGizmoActor();
 											if( bUsingGizmo )
 											{
 												ClickedTransformGizmoComponent = HitResult.GetComponent();
@@ -317,14 +342,18 @@ bool UViewportInteractor::HandleInputKey( FEditorViewportClient& ViewportClient,
 										}
 
 										const bool bIsPlacingNewObjects = false;
-										const bool bStartTransaction = true;
-										if( WorldInteraction->StartDragging( this, ClickedTransformGizmoComponent, HitResult.ImpactPoint, bIsPlacingNewObjects, bStartTransaction ) )
-										{
-											Action->bIsInputCaptured = true;
-										}
+										const bool bAllowInterpolationWhenPlacing = true;
+										const bool bStartTransaction = !bSelectionChanged;
+										const bool bWithGrabberSphere = false;	// @todo grabber: Not supported yet
+										WorldInteraction->StartDragging( this, ClickedTransformGizmoComponent, HitResult.ImpactPoint, bIsPlacingNewObjects, bAllowInterpolationWhenPlacing, bStartTransaction, bWithGrabberSphere );
+
+										Action->bIsInputCaptured = true;
 									}
 									else if ( bSelectionChanged )
 									{
+										// Stop our transaction 
+										TrackingTransaction.End();
+
 										// User didn't drag but did select something, so play a haptic effect.
 										PlayHapticEffect( VI::SelectionHapticFeedbackStrength->GetFloat() );
 									}
@@ -380,9 +409,11 @@ bool UViewportInteractor::HandleInputKey( FEditorViewportClient& ViewportClient,
 					}
 
 					InteractorData.bIsFirstDragUpdate = true;
+					InteractorData.bDraggingWithGrabberSphere = false;
 					InteractorData.DragRayLength = 0.0f;
-					InteractorData.LastLaserPointerStart = InteractorData.Transform.GetLocation();
 					InteractorData.LastDragToLocation = InteractorData.Transform.GetLocation();
+					InteractorData.InteractorRotationAtDragStart = InteractorData.Transform.GetRotation();
+					InteractorData.GrabberSphereLocationAtDragStart = FVector::ZeroVector;
 					InteractorData.DragTranslationVelocity = FVector::ZeroVector;
 					InteractorData.DragRayLengthVelocity = 0.0f;
 					InteractorData.bIsDrivingVelocityOfSimulatedTransformables = false;
@@ -392,6 +423,7 @@ bool UViewportInteractor::HandleInputKey( FEditorViewportClient& ViewportClient,
 					InteractorData.TransformGizmoInteractionType = ETransformGizmoInteractionType::None;
 					InteractorData.OptionalHandlePlacement.Reset();
 					InteractorData.GizmoStartTransform = FTransform::Identity;
+					InteractorData.GizmoLastTransform = InteractorData.GizmoTargetTransform = InteractorData.GizmoUnsnappedTargetTransform = InteractorData.GizmoInterpolationSnapshotTransform = InteractorData.GizmoStartTransform;
 					InteractorData.GizmoStartLocalBounds = FBox( 0 );
 					InteractorData.GizmoSpaceFirstDragUpdateOffsetAlongAxis = FVector::ZeroVector;
 					InteractorData.GizmoSpaceDragDeltaFromStartOffset = FVector::ZeroVector;
@@ -433,8 +465,11 @@ bool UViewportInteractor::HandleInputKey( FEditorViewportClient& ViewportClient,
 
 		if ( !bHandled )
 		{
+			// Determine whether we are starting a marquee select
+			const bool bIsMarqueeSelect = ViewportClient.IsAltPressed() && ViewportClient.IsCtrlPressed();
+
 			// If "select and move" was pressed but not handled, go ahead and deselect everything
-			if ( Action->ActionType == ViewportWorldActionTypes::SelectAndMove && Event == IE_Pressed )
+			if ( !bIsMarqueeSelect && Action->ActionType == ViewportWorldActionTypes::SelectAndMove && Event == IE_Pressed )
 			{
 				WorldInteraction->Deselect();
 			}
@@ -454,13 +489,27 @@ bool UViewportInteractor::HandleInputAxis( FEditorViewportClient& ViewportClient
 	{
 		FViewportActionKeyInput Action( KnownAction->ActionType );
 
-		// Give the derived classes a chance to update according to the input
-		HandleInputAxis( ViewportClient, Action, Key, Delta, DeltaTime, bHandled );
+		if( !bHandled )
+		{
+			// Give the derived classes a chance to update according to the input
+			PreviewInputAxis( ViewportClient, Action, Key, Delta, DeltaTime, bHandled );
+		}
+
+		if( !bHandled )
+		{
+			// Give the derived classes a chance to update according to the input
+			HandleInputAxis( ViewportClient, Action, Key, Delta, DeltaTime, bHandled );
+		}
 	}
 
 	return bHandled;
 }
 
+bool UViewportInteractor::AllowLaserSmoothing() const
+{
+	return true;
+}
+
 FTransform UViewportInteractor::GetTransform() const
 {
 	return InteractorData.Transform;
@@ -481,35 +530,52 @@ FVector UViewportInteractor::GetDragTranslationVelocity() const
 	return InteractorData.DragTranslationVelocity;
 }
 
-bool UViewportInteractor::GetLaserPointer( FVector& LaserPointerStart, FVector& LaserPointerEnd, const bool bEvenIfBlocked, const float LaserLengthOverride ) const
+void UViewportInteractor::SetHoverLocation(const FVector& InHoverLocation)
+{
+	InteractorData.HoverLocation = InHoverLocation;
+}
+
+bool UViewportInteractor::GetLaserPointer( FVector& LaserPointerStart, FVector& LaserPointerEnd, const bool bEvenIfBlocked, const float LaserLengthOverride )
 {
-	// If we're currently grabbing the world with both hands, then the laser pointer is not available
-	if ( /*bHaveMotionController &&*/ //@todo ViewportInteraction
-		 OtherInteractor == nullptr ||
-		 ( !( InteractorData.DraggingMode == EViewportInteractionDraggingMode::World && GetOtherInteractor()->GetInteractorData().DraggingMode == EViewportInteractionDraggingMode::AssistingDrag ) &&
-		   !( InteractorData.DraggingMode == EViewportInteractionDraggingMode::AssistingDrag && GetOtherInteractor()->GetInteractorData().DraggingMode == EViewportInteractionDraggingMode::World ) ) )
+	// If we have UI attached to us, don't allow a laser pointer
+	if( bEvenIfBlocked || !GetIsLaserBlocked() )
 	{
-		// If we have UI attached to us, don't allow a laser pointer
-		if ( bEvenIfBlocked || !GetIsLaserBlocked() )
+		FTransform HandTransform;
+		FVector HandForwardVector;
+		if (GetTransformAndForwardVector(HandTransform, HandForwardVector))
 		{
-			FTransform HandTransform;
-			FVector HandForwardVector;
-			if ( GetTransformAndForwardVector( HandTransform, HandForwardVector ) )
+			LaserPointerStart = HandTransform.GetLocation();
+
+			/** To avoid calculating the smooth end location multiple times in one tick we will check if it has already been done this frame and use that */
+			if (!SavedLaserPointerEnd.IsSet())
 			{
-				LaserPointerStart = HandTransform.GetLocation();
+				const float LaserLength = LaserLengthOverride == 0.0f ? GetLaserPointerMaxLength() : LaserLengthOverride;
+				FVector FinalLaserPointerEnd = LaserPointerStart + HandForwardVector * LaserLength;
 
-				float LaserLength = LaserLengthOverride == 0.0f ? GetLaserPointerMaxLength() : LaserLengthOverride;
-				LaserPointerEnd = LaserPointerStart + HandForwardVector * LaserLength;
+				/** Only smooth the endlocation when not disabled */
+				if (WorldInteraction != nullptr && AllowLaserSmoothing() && VI::AllowLaserSmooth->GetInt() == 1)
+				{
+					SmoothingOneEuroFilter.SetCutoffSlope(VI::LaserSmoothLag->GetFloat());
+					SmoothingOneEuroFilter.SetMinCutoff(VI::LaserSmoothMinimumCutoff->GetFloat());
+					FinalLaserPointerEnd = SmoothingOneEuroFilter.Filter(FinalLaserPointerEnd, WorldInteraction->GetCurrentDeltaTime());
+				}
 
-				return true;
+				LaserPointerEnd = FinalLaserPointerEnd;
+				SavedLaserPointerEnd = FinalLaserPointerEnd;
+			}
+			else
+			{
+				LaserPointerEnd = SavedLaserPointerEnd.GetValue();
 			}
+
+			return true;
 		}
 	}
 	return false;
 }
 
 
-bool UViewportInteractor::GetGrabberSphere( FSphere& OutGrabberSphere, const bool bEvenIfBlocked ) const
+bool UViewportInteractor::GetGrabberSphere( FSphere& OutGrabberSphere, const bool bEvenIfBlocked )
 {
 	OutGrabberSphere = FSphere( 0 );
 
@@ -536,7 +602,13 @@ bool UViewportInteractor::GetGrabberSphere( FSphere& OutGrabberSphere, const boo
 
 float UViewportInteractor::GetLaserPointerMaxLength() const
 {
-	return VI::LaserPointerMaxLength->GetFloat() * WorldInteraction->GetWorldScaleFactor();
+	return VI::LaserPointerMaxLength->GetFloat();
+}
+
+void UViewportInteractor::ResetHoverState()
+{
+	InteractorData.HoverLocation.Reset();
+	InteractorData.HoveringOverTransformGizmoComponent = nullptr;
 }
 
 FHitResult UViewportInteractor::GetHitResultFromLaserPointer( TArray<AActor*>* OptionalListOfIgnoredActors /*= nullptr*/, const bool bIgnoreGizmos /*= false*/, 
@@ -553,6 +625,21 @@ FHitResult UViewportInteractor::GetHitResultFromLaserPointer( TArray<AActor*>* O
 			bActuallyIgnoreGizmos = true;
 		}
 
+		// Ignore all volume objects.  They'll just get in the way of selecting other things.
+		// @todo viewportinteraction: We'll need to device a way to allow volume wire bounds to be selectable using this system
+		static TArray<AActor*> VolumeActors;
+		VolumeActors.Reset();
+		{
+			for( TActorIterator<AVolume> It( WorldInteraction->GetWorld(), AVolume::StaticClass() ); It; ++It )
+			{
+				AActor* Actor = *It;
+				if( !Actor->IsPendingKill() )
+				{
+					VolumeActors.Add( Actor );
+				}
+			}
+		}
+
 		// Twice twice.  Once for editor gizmos which are "on top" and always take precedence, then a second time
 		// for all of the scene objects
 		for ( int32 PassIndex = bActuallyIgnoreGizmos ? 1 : 0; PassIndex < 2; ++PassIndex )
@@ -567,6 +654,17 @@ FHitResult UViewportInteractor::GetHitResultFromLaserPointer( TArray<AActor*>* O
 				TraceParams.AddIgnoredActors( *OptionalListOfIgnoredActors );
 			}
 
+			for( const TWeakObjectPtr<AActor> ActorToIgnoreWeakPtr : WorldInteraction->GetActorsToExcludeFromHitTest() )
+			{
+				AActor* ActorToIgnore = ActorToIgnoreWeakPtr.Get();
+				if( ActorToIgnore != nullptr )
+				{
+					TraceParams.AddIgnoredActor( ActorToIgnore );
+				}
+			}
+
+			TraceParams.AddIgnoredActors( VolumeActors );
+
 			bool bHit = false;
 			FHitResult HitResult;
 			if ( bOnlyEditorGizmos )
@@ -629,14 +727,26 @@ bool UViewportInteractor::GetTransformAndForwardVector( FTransform& OutHandTrans
 	return true;
 }
 
-FVector UViewportInteractor::GetHoverLocation() const 
+FVector UViewportInteractor::GetHoverLocation() 
 {
-	return InteractorData.HoverLocation;
+	FVector Result = FVector::ZeroVector;
+	if (InteractorData.HoverLocation.IsSet())
+	{
+		Result = InteractorData.HoverLocation.GetValue();
+	}
+	else
+	{
+		FVector LaserStart, LaserEnd;
+		GetLaserPointer(LaserStart, LaserEnd, true);
+		Result = LaserEnd;
+	}
+
+	return Result;
 }
 
 bool UViewportInteractor::IsHovering() const
 {
-	return InteractorData.bIsHovering;
+	return InteractorData.HoverLocation.IsSet();
 }
 
 bool UViewportInteractor::IsHoveringOverGizmo()
@@ -668,26 +778,36 @@ float UViewportInteractor::GetDragHapticFeedbackStrength() const
 	return VI::DragHapticFeedbackStrength->GetFloat();
 }
 
-void UViewportInteractor::SetAllowTriggerLightPressLocking( const bool bInAllow )
+bool UViewportInteractor::IsHoveringOverPriorityType() const
 {
-	InteractorData.bAllowTriggerLightPressLocking = bInAllow;
+	return InteractorData.bHitResultIsPriorityType;
 }
 
-bool UViewportInteractor::AllowTriggerLightPressLocking() const
+bool UViewportInteractor::IsHoveringOverSelectedActor() const
 {
-	return InteractorData.bAllowTriggerLightPressLocking;
-}
+	bool bResult = false;
 
-void UViewportInteractor::SetAllowTriggerFullPress( const bool bInAllow )
-{
-	InteractorData.bAllowTriggerFullPress = bInAllow;
-}
+	if (InteractorData.LastHoveredActorComponent != nullptr && InteractorData.HoverLocation.IsSet())
+	{
+		USelection* SelectedActors = GEditor->GetSelectedActors();
+		for (FSelectionIterator SelectionIt(*SelectedActors); SelectionIt; ++SelectionIt)
+		{
+			AActor* Actor = CastChecked<AActor>(*SelectionIt);
+			if (InteractorData.LastHoveredActorComponent->GetOwner() == Actor)
+			{
+				bResult = true;
+				break;
+			}
+		}
+	}
 
-bool UViewportInteractor::AllowTriggerFullPress() const
-{
-	return InteractorData.bAllowTriggerFullPress;
+	return bResult;
 }
-bool UViewportInteractor::IsHoveringOverPriorityType() const
+
+
+void UViewportInteractor::ResetLaserEnd()
 {
-	return InteractorData.bHitResultIsPriorityType;
+	SavedLaserPointerEnd.Reset();
 }
+
+#undef LOCTEXT_NAMESPACE
diff --git a/Engine/Source/Editor/ViewportInteraction/ViewportInteractorData.h b/Engine/Source/Editor/ViewportInteraction/ViewportInteractorData.h
--- a/Engine/Source/Editor/ViewportInteraction/ViewportInteractorData.h
+++ b/Engine/Source/Editor/ViewportInteraction/ViewportInteractorData.h
@@ -30,17 +30,18 @@ struct FViewportInteractorData
 	// Hover feedback
 	//
 
-	/** Lights up when our laser pointer is close enough to something to click */
-	bool bIsHovering;
-
-	/** The widget component we are currently hovering over */
-	class UWidgetComponent* HoveringOverWidgetComponent; //@todo: ViewportInteraction: UI should not be in this module
-
+	/** The widget component we last hovered over.  This is used to detect when the laser pointer moves over or leaves a widget, and is not reset every frame */
+	class UWidgetComponent* LastHoveredWidgetComponent; //@todo: ViewportInteraction: UI should not be in this module
+		
 	/** Position the laser pointer impacted an interactive object at (UI, meshes, etc.) */
-	FVector HoverLocation;
+	TOptional<FVector> HoverLocation;
 
 	/** The current component hovered by the laser pointer of this hand */
-	TWeakObjectPtr<class UActorComponent> HoveredActorComponent;
+	TWeakObjectPtr<class UActorComponent> LastHoveredActorComponent;
+
+	/** The last location that we hovered over UI at in the world.  This is used for dragging and dropping from UI that
+	    may have already been closed, such as Content Browser */
+	FVector LastHoverLocationOverUI;
 
 	//
 	// General input
@@ -49,9 +50,6 @@ struct FViewportInteractorData
 	/** True if we're currently holding the 'SelectAndMove' button down after clicking on an actor */
 	TWeakObjectPtr<UActorComponent> ClickingOnComponent;
 
-	/** Last real time that we pressed the 'SelectAndMove' button on UI.  This is used to detect double-clicks. */
-	double LastClickPressTime;
-
 
 	//
 	// Object/world movement
@@ -63,6 +61,9 @@ struct FViewportInteractorData
 	/** What we were doing last.  Used for inertial movement. */
 	EViewportInteractionDraggingMode LastDraggingMode;
 
+	/** True if we're dragging using the grabber sphere, or false if we're using the laser (or world movement) */
+	bool bDraggingWithGrabberSphere;
+
 	/** True if this is the first update since we started dragging */
 	bool bIsFirstDragUpdate;
 
@@ -72,14 +73,17 @@ struct FViewportInteractorData
 	/** Length of the ray that's dragging */
 	float DragRayLength;
 
-	/** Laser pointer start location the last frame */
-	FVector LastLaserPointerStart;
-
 	/** Location that we dragged to last frame (end point of the ray) */
 	FVector LastDragToLocation;
 
-	/** Laser pointer impact location at the drag start */
-	FVector LaserPointerImpactAtDragStart;
+	/** The orientation of the interactor when we first started the drag */
+	FQuat InteractorRotationAtDragStart;
+
+	/** Where the grabber sphere center point was when we first started the drag */
+	FVector GrabberSphereLocationAtDragStart;
+
+	/** Grabber sphere or laser pointer impact location at the drag start */
+	FVector ImpactLocationAtDragStart;
 
 	/** How fast to move selected objects every frame for inertial translation */
 	FVector DragTranslationVelocity;
@@ -99,6 +103,18 @@ struct FViewportInteractorData
 	/** Where the gizmo was placed at the beginning of the current interaction */
 	FTransform GizmoStartTransform;
 
+	/** Where the gizmo was last frame.  This is used for interpolation and smooth snapping */
+	FTransform GizmoLastTransform;
+
+	/** Where the gizmo wants to be right now, with snaps applied. */
+	FTransform GizmoTargetTransform;
+
+	/** Where the gizmo wants to be right now, if no snaps were applied.  This is used for interpolation and smooth snapping */
+	FTransform GizmoUnsnappedTargetTransform;
+
+	/** A transform that we're interpolating from, toward the target transform.  This is used when placing objects, so they'll smoothly interpolate to their initial location. */
+	FTransform GizmoInterpolationSnapshotTransform;
+
 	/** Our gizmo bounds at the start of the interaction, in actor local space. */
 	FBox GizmoStartLocalBounds;
 
@@ -124,17 +140,6 @@ struct FViewportInteractorData
 	/** Gizmo handle that we hovered over last (used only for avoiding spamming of hover haptics!) */
 	TWeakObjectPtr<class USceneComponent> HoverHapticCheckLastHoveredGizmoComponent;
 
-	//@todo ViewportInteraction: Light press should not have to be here, ViewportInteractor needs it
-	/** True if we allow locking of the lightly pressed state for the current event.  This can be useful to
-	turn off in cases where you always want a full press to override the light press, even if it was
-	a very slow press */
-	bool bAllowTriggerLightPressLocking;
-
-	/** This can be set to false when handling SelectAndMove_LightlyPressed to prevent a full press event from
-	    being possible at all.  Useful when you're not planning to use FullyPressed for anything, but you
-		don't want a full press to cancel your light press. */
-	bool bAllowTriggerFullPress;
-
 	/** If the latest hitresult is hovering over a priority type */
 	bool bHitResultIsPriorityType;
 
@@ -145,26 +150,29 @@ struct FViewportInteractorData
 	/** Default constructor for FVirtualHand that initializes safe defaults */
 	FViewportInteractorData()
 	{
-		bIsHovering = false;
-		HoveringOverWidgetComponent = nullptr;
-		HoverLocation = FVector::ZeroVector;
-		HoveredActorComponent = nullptr;
+		LastHoveredWidgetComponent = nullptr;
+		HoverLocation = TOptional<FVector>();
+		LastHoveredActorComponent = nullptr;
+		LastHoverLocationOverUI = FVector::ZeroVector;
 
 		ClickingOnComponent = nullptr;
 
 		DraggingMode = EViewportInteractionDraggingMode::Nothing;
+		bDraggingWithGrabberSphere = false;
 		LastDraggingMode = EViewportInteractionDraggingMode::Nothing;
 		bWasAssistingDrag = false;
 		bIsFirstDragUpdate = false;
 		DragRayLength = 0.0f;
-		LastLaserPointerStart = FVector::ZeroVector;
 		LastDragToLocation = FVector::ZeroVector;
-		LaserPointerImpactAtDragStart = FVector::ZeroVector;
+		InteractorRotationAtDragStart = FQuat::Identity;
+		GrabberSphereLocationAtDragStart = FVector::ZeroVector;
+		ImpactLocationAtDragStart = FVector::ZeroVector;
 		DragTranslationVelocity = FVector::ZeroVector;
 		DragRayLengthVelocity = 0.0f;
 		bIsDrivingVelocityOfSimulatedTransformables = false;
 
 		GizmoStartTransform = FTransform::Identity;
+		GizmoLastTransform = GizmoTargetTransform = GizmoUnsnappedTargetTransform = GizmoInterpolationSnapshotTransform = GizmoStartTransform;
 		GizmoStartLocalBounds = FBox( 0 );
 		GizmoSpaceFirstDragUpdateOffsetAlongAxis = FVector::ZeroVector;
 		GizmoSpaceDragDeltaFromStartOffset = FVector::ZeroVector;
@@ -174,8 +182,6 @@ struct FViewportInteractorData
 		HoveringOverTransformGizmoComponent = nullptr;
 		HoverHapticCheckLastHoveredGizmoComponent = nullptr;
 
-		bAllowTriggerLightPressLocking = true;
-		bAllowTriggerFullPress = true;
 		bHitResultIsPriorityType = false;
 
 		StartHitLocationToTransformableCenter = FVector::ZeroVector;
