Merging //UE4/Private-Geometry to Dev-VREditor (//UE4/Dev-VREditor) 3339870

diff --git a/Engine/Source/Editor/ViewportInteraction/ViewportWorldInteraction.cpp b/Engine/Source/Editor/ViewportInteraction/ViewportWorldInteraction.cpp
--- a/Engine/Source/Editor/ViewportInteraction/ViewportWorldInteraction.cpp
+++ b/Engine/Source/Editor/ViewportInteraction/ViewportWorldInteraction.cpp
@@ -26,21 +26,34 @@
 #include "SnappingUtils.h"
 #include "ScopedTransaction.h"
 
+#include "DrawDebugHelpers.h"
+
+//Sound
+#include "Kismet/GameplayStatics.h"
+#include "ViewportInteractionStyle.h"
+#include "ViewportInteractionAssetContainer.h"
+
 // For actor placement
 #include "IHeadMountedDisplay.h"
+#include "EngineUtils.h"
+#include "ActorViewportTransformable.h"
 
 #define LOCTEXT_NAMESPACE "ViewportWorldInteraction"
 
 namespace VI
 {
-	static FAutoConsoleVariable ScaleSensitivity( TEXT( "VI.ScaleSensitivity" ), 0.01f, TEXT( "Sensitivity for scaling" ) );
+	static FAutoConsoleVariable GizmoScaleInDesktop( TEXT( "VI.GizmoScaleInDesktop" ), 0.35f, TEXT( "How big the transform gizmo should be when used in desktop mode" ) );
+	static FAutoConsoleVariable ScaleSensitivity( TEXT( "VI.ScaleSensitivity" ), 0.005f, TEXT( "Sensitivity for scaling" ) );
 	static FAutoConsoleVariable GizmoRotationSensitivity( TEXT( "VI.GizmoRotationSensitivity" ), 0.25f, TEXT( "How much to rotate as the user drags on a rotation gizmo handle" ) );
 	static FAutoConsoleVariable ScaleWorldFromFloor( TEXT( "VI.ScaleWorldFromFloor" ), 0, TEXT( "Whether the world should scale relative to your tracking space floor instead of the center of your hand locations" ) );
+	static FAutoConsoleVariable ScaleWorldWithDynamicPivot( TEXT( "VI.ScaleWorldWithDynamicPivot" ), 1, TEXT( "Whether to compute a new center point for scaling relative from by looking at how far either controller moved relative to the last frame" ) );
 	static FAutoConsoleVariable AllowVerticalWorldMovement( TEXT( "VI.AllowVerticalWorldMovement" ), 1, TEXT( "Whether you can move your tracking space away from the origin or not" ) );
 	static FAutoConsoleVariable AllowWorldRotationPitchAndRoll( TEXT( "VI.AllowWorldRotationPitchAndRoll" ), 0, TEXT( "When enabled, you'll not only be able to yaw, but also pitch and roll the world when rotating by gripping with two hands" ) );
-	static FAutoConsoleVariable InertiaVelocityBoost( TEXT( "VI.InertiaVelocityBoost" ), 3000.0f, TEXT( "How much to scale object velocity when releasing dragged simulating objects in Simulate mode" ) );
-	static FAutoConsoleVariable SweepPhysicsWhileSimulating( TEXT( "VI.SweepPhysicsWhileSimulating" ), 1, TEXT( "If enabled, simulated objects won't be able to penetrate other objects while being dragged in Simulate mode" ) );
+	static FAutoConsoleVariable InertiaVelocityBoost( TEXT( "VI.InertiaVelocityBoost" ), 0.5f, TEXT( "How much to scale object velocity when releasing dragged simulating objects in Simulate mode" ) );
+	static FAutoConsoleVariable SweepPhysicsWhileSimulating( TEXT( "VI.SweepPhysicsWhileSimulating" ), 0, TEXT( "If enabled, simulated objects won't be able to penetrate other objects while being dragged in Simulate mode" ) );
 	static FAutoConsoleVariable PlacementInterpolationDuration( TEXT( "VI.PlacementInterpolationDuration" ), 0.6f, TEXT( "How long we should interpolate newly-placed objects to their target location." ) );
+	static FAutoConsoleVariable PlacementOffsetScaleWhileSimulating( TEXT( "VI.PlacementOffsetScaleWhileSimulating" ), 0.25f, TEXT( "How far to additionally offset objects (as a scalar percentage of the gizmo bounds) from the placement impact point while simulate mode is active" ) );
+
 	static FAutoConsoleVariable SmoothSnap( TEXT( "VI.SmoothSnap" ), 1, TEXT( "When enabled with grid snap, transformed objects will smoothly blend to their new location (instead of teleporting instantly)" ) );
 	static FAutoConsoleVariable SmoothSnapSpeed( TEXT( "VI.SmoothSnapSpeed" ), 30.0f, TEXT( "How quickly objects should interpolate to their new position when grid snapping is enabled" ) );
 	static FAutoConsoleVariable ElasticSnap( TEXT( "VI.ElasticSnap" ), 1, TEXT( "When enabled with grid snap, you can 'pull' objects slightly away from their snapped position" ) );
@@ -55,14 +68,35 @@ namespace VI
 	static FAutoConsoleVariable GizmoHandleHoverAnimationDuration( TEXT( "VI.GizmoHandleHoverAnimationDuration" ), 0.1f, TEXT( "How quickly to animate gizmo handle hover state" ) );
 	static FAutoConsoleVariable ShowTransformGizmo( TEXT( "VI.ShowTransformGizmo" ), 1, TEXT( "Whether the transform gizmo should be shown for selected objects" ) );
 	static FAutoConsoleVariable DragTranslationVelocityStopEpsilon( TEXT( "VI.DragTranslationVelocityStopEpsilon" ), KINDA_SMALL_NUMBER, TEXT( "When dragging inertia falls below this value (cm/frame), we'll stop inertia and finalize the drag" ) );
-	static FAutoConsoleVariable SnapGridSize( TEXT( "VI.SnapGridSize" ), 3.0f, TEXT( "How big the snap grid should be, in multiples of the gizmo bounding box size" ) );
+	static FAutoConsoleVariable SnapGridSize( TEXT( "VI.SnapGridSize" ), 3.0f, TEXT( "How big the snap grid should be.  At 1.0, this will be the maximum of the gizmo's bounding box and a multiple of the current grid snap size" ) );
 	static FAutoConsoleVariable SnapGridLineWidth( TEXT( "VI.SnapGridLineWidth" ), 3.0f, TEXT( "Width of the grid lines on the snap grid" ) );
 	static FAutoConsoleVariable MinVelocityForInertia( TEXT( "VI.MinVelocityForInertia" ), 1.0f, TEXT( "Minimum velocity (in cm/frame in unscaled room space) before inertia will kick in when releasing objects (or the world)" ) );
 	static FAutoConsoleVariable GridHapticFeedbackStrength( TEXT( "VI.GridHapticFeedbackStrength" ), 0.4f, TEXT( "Default strength for haptic feedback when moving across grid points" ) );
 
 	static FAutoConsoleVariable UseTransientActors( TEXT( "VI.UseTransientActors" ), 1, TEXT( "For debugging only, allows you to turn off transient on newly-spawned VR editor actors." ) );
+
+	static FAutoConsoleVariable EnableGuides(TEXT("VI.EnableGuides"), 0, TEXT("Whether or not guidelines should be enabled. Off by default, set to 1 to enable."));
+	static FAutoConsoleVariable AlignCandidateDistance(TEXT("VI.AlignCandidateDistance"), 2.0f, TEXT("The distance candidate actors can be from our transformable (in multiples of our transformable's size"));
+	static FAutoConsoleVariable ForceSnapDistance(TEXT("VI.ForceSnapDistance"), 25.0f, TEXT("The distance (in % of transformable size) where guide lines indicate that actors are aligned"));
+
+	static FAutoConsoleVariable ForceShowCursor(TEXT("VI.ForceShowCursor"), 0, TEXT("Whether or not the mirror window's cursor should be enabled. Off by default, set to 1 to enable."));
+	static FAutoConsoleVariable SFXMultiplier(TEXT("VI.SFXMultiplier"), 1.5f, TEXT("Default Sound Effect Volume Multiplier"));
 }
 
+const FString UViewportWorldInteraction::AssetContainerPath = FString("/Engine/VREditor/ViewportInteractionAssetContainerData");
+
+struct FGuideData
+{
+	const AActor* AlignedActor;
+	FVector LocalOffset;
+	FVector SnapPoint;
+	FVector GuideStart;
+	FVector GuideEnd;
+	FColor GuideColor;
+	float DrawAlpha;
+	float GuideLength;
+};
+
 // @todo vreditor: Hacky inline implementation of a double vector.  Move elsewhere and polish or get rid.
 struct DVector
 {
@@ -232,25 +266,34 @@ UViewportWorldInteraction::UViewportWorldInteraction():
 	MotionControllerID( 0 ),	// @todo ViewportInteraction: We only support a single controller, and we assume the first controller are the motion controls
 	LastWorldToMetersScale( 100.0f ),
 	bSkipInteractiveWorldMovementThisFrame( false ),
+	RoomTransformToSetOnFrame(),
 	bAreTransformablesMoving( false ),
 	bIsInterpolatingTransformablesFromSnapshotTransform( false ),
+	bFreezePlacementWhileInterpolatingTransformables( false ),
 	TransformablesInterpolationStartTime( FTimespan::Zero() ),
 	TransformablesInterpolationDuration( 1.0f ),
 	TransformGizmoActor( nullptr ),
 	TransformGizmoClass( APivotTransformGizmo::StaticClass() ),
 	GizmoLocalBounds( FBox( 0 ) ),
 	StartDragAngleOnRotation(),
-	StartDragHandleDirection(),
+	DraggingRotationHandleDirection(),
 	bShouldTransformGizmoBeVisible( true ),
-	TransformGizmoScale( 0.35f ),
+	TransformGizmoScale( VI::GizmoScaleInDesktop->GetFloat() ),
 	SnapGridActor( nullptr ),
 	SnapGridMeshComponent( nullptr ),
 	SnapGridMID( nullptr ),
 	DraggedInteractable( nullptr ),
 	DefaultMouseCursorInteractor( nullptr ),
 	DefaultMouseCursorInteractorRefCount( 0 ),
+	bIsInVR( false ),
 	bActive( false ),
-	bUseInputProcessor( false )
+	bUseInputPreprocessor( false ),
+	bAllowWorldMovement( true ),
+	CurrentDeltaTime(0.0f),
+	bShouldSuppressCursor(false),
+	CurrentTickNumber(0),
+	AssetContainer(nullptr),
+	bPlayNextRefreshTransformGizmoSound(true)
 {
 }
 
@@ -259,16 +302,20 @@ void UViewportWorldInteraction::Init()
 {
 	Colors.SetNumZeroed( (int32)EColors::TotalCount );
 	{
-		Colors[ (int32)EColors::DefaultColor ] = FLinearColor( 0.7f, 0.7f, 0.7f, 1.0f );;
-		Colors[ (int32)EColors::Forward ] = FLinearColor( 0.4f, 0.05f, 0.05f, 1.0f );
-		Colors[ (int32)EColors::Right ] = FLinearColor( 0.05f, 0.4f, 0.05f, 1.0f );
-		Colors[ (int32)EColors::Up ] = FLinearColor( 0.05f, 0.05f, 0.4f, 1.0f );
-		Colors[ (int32)EColors::Hover ] = FLinearColor::Yellow;
-		Colors[ (int32)EColors::Dragging ] = FLinearColor::White;
+		Colors[(int32)EColors::DefaultColor] = FLinearColor(0.7f, 0.7f, 0.7f, 1.0f);
+		Colors[(int32)EColors::Forward] = FLinearColor(0.594f, 0.0197f, 0.0f, 1.0f);
+		Colors[(int32)EColors::Right] = FLinearColor(0.1349f, 0.3959f, 0.0f, 1.0f);
+		Colors[(int32)EColors::Up] = FLinearColor(0.0251f, 0.207f, 0.85f, 1.0f);
+		Colors[(int32)EColors::GizmoHover] = FLinearColor::Yellow;
+		Colors[(int32)EColors::GizmoDragging] = FLinearColor::Yellow;
 	}
 
 	AppTimeEntered = FTimespan::FromSeconds( FApp::GetCurrentTime() );
 
+	// Setup the asset container.
+	AssetContainer = LoadObject<UViewportInteractionAssetContainer>(nullptr, *UViewportWorldInteraction::AssetContainerPath);
+	check(AssetContainer != nullptr);
+
 	IViewportWorldInteractionManager& WorldInteractionManager = IViewportInteractionModule::Get().GetWorldInteractionManager();
 	WorldInteractionManager.SetCurrentViewportWorldInteraction( this );
 
@@ -277,19 +324,37 @@ void UViewportWorldInteraction::Init()
 
 	//Spawn the transform gizmo at init so we do not hitch when selecting our first object
 	SpawnTransformGizmoIfNeeded();
-	TransformGizmoActor->SetIsTemporarilyHiddenInEditor( true );
+
+	const bool bShouldBeVisible = false;
+	const bool bPropagateToChildren = true;
+	TransformGizmoActor->GetRootComponent()->SetVisibility( bShouldBeVisible, bPropagateToChildren );
 
 
 	/** This will make sure this is not ticking after the editor has been closed. */
 	GEditor->OnEditorClose().AddUObject( this, &UViewportWorldInteraction::Shutdown );
 
+	// We need a mouse cursor!
+	this->AddMouseCursorInteractor();
+
 	bActive = true;
+
+	CandidateActors.Reset();
+
+	// Pretend that actor selection changed, so that our gizmo refreshes right away based on which objects are selected
+	GEditor->NoteSelectionChange();
+
+	CurrentTickNumber = 0;
 }
 
 void UViewportWorldInteraction::Shutdown()
 {
 	bActive = false;
 
+	if( DefaultMouseCursorInteractorRefCount == 1 )
+	{
+		this->ReleaseMouseCursorInteractor();
+	}
+
 	DestroyActors();
 	DefaultOptionalViewportClient = nullptr;
 
@@ -309,6 +374,7 @@ void UViewportWorldInteraction::Shutdown()
 	Colors.Empty();
 
 	OnHoverUpdateEvent.Clear();
+	OnPreviewInputActionEvent.Clear();
 	OnInputActionEvent.Clear();
 	OnKeyInputEvent.Clear();
 	OnAxisInputEvent.Clear();
@@ -325,10 +391,34 @@ void UViewportWorldInteraction::Shutdown()
 		ViewportTransformer = nullptr;
 	}
 
+	AssetContainer = nullptr;
+
 	USelection::SelectionChangedEvent.RemoveAll( this );
 	GEditor->OnEditorClose().RemoveAll( this );
 }
 
+void UViewportWorldInteraction::TransitionWorld(UWorld* NewWorld)
+{
+	Super::TransitionWorld(NewWorld);
+
+	for (UViewportInteractor* Interactor : Interactors)
+	{
+		Interactor->Rename(nullptr, NewWorld->PersistentLevel);
+	}
+}
+
+void UViewportWorldInteraction::EnteredSimulateInEditor()
+{
+	// Make sure transformables get updated
+	GEditor->NoteSelectionChange();
+}
+
+void UViewportWorldInteraction::LeftSimulateInEditor()
+{
+	// Make sure transformables get updated
+	GEditor->NoteSelectionChange();
+}
+
 void UViewportWorldInteraction::OnEditorClosed()
 {
 	if( bActive )
@@ -339,10 +429,20 @@ void UViewportWorldInteraction::OnEditorClosed()
 
 void UViewportWorldInteraction::Tick( float DeltaSeconds )
 {
-	OnPreWorldInteractionTickEvent.Broadcast( DeltaSeconds );
+	CurrentDeltaTime = DeltaSeconds;
+
+	if( RoomTransformToSetOnFrame.IsSet() )
+	{
+		const uint32 SetOnFrameNumber = RoomTransformToSetOnFrame.GetValue().Get<1>();
+		if( SetOnFrameNumber == CurrentTickNumber )
+		{
+			const FTransform& NewRoomTransform = RoomTransformToSetOnFrame.GetValue().Get<0>();
+			SetRoomTransform( NewRoomTransform );
+			RoomTransformToSetOnFrame.Reset();
+		}
+	}
 
-	// Cancel any haptic effects that have been playing too long
-	//StopOldHapticEffects(); //@todo ViewportInteraction
+	OnPreWorldInteractionTickEvent.Broadcast( DeltaSeconds );
 
 	// Get the latest interactor data, and fill in our its data with fresh transforms
 	PollInputIfNeeded();
@@ -359,6 +459,13 @@ void UViewportWorldInteraction::Tick( float DeltaSeconds )
 	}
 
 	OnPostWorldInteractionTickEvent.Broadcast( DeltaSeconds );
+
+	for (UViewportInteractor* Interactor : Interactors)
+	{
+		Interactor->ResetLaserEnd();
+	}
+
+	CurrentTickNumber++;
 }
 
 void UViewportWorldInteraction::AddInteractor( UViewportInteractor* Interactor )
@@ -403,16 +510,22 @@ void UViewportWorldInteraction::SetTransformables( TArray< TUniquePtr< class FVi
 	bDraggedSinceLastSelection = false;
 	LastDragGizmoStartTransform = FTransform::Identity;
 	bIsInterpolatingTransformablesFromSnapshotTransform = false;
+	bFreezePlacementWhileInterpolatingTransformables = false;
 	TransformablesInterpolationStartTime = FTimespan::Zero();
 	TransformablesInterpolationDuration = 1.0f;
 
+	// Clear our last dragging mode on all interactors, so no inertia from the last objects we dragged will
+	// be applied to the new transformables.
+	for( UViewportInteractor* Interactor : Interactors )
+	{
+		FViewportInteractorData& InteractorData = Interactor->GetInteractorData();
+		InteractorData.LastDraggingMode = EViewportInteractionDraggingMode::Nothing;
+	}
+
 	Transformables = MoveTemp( NewTransformables );
 
 	const bool bNewObjectsSelected = true;
-	const bool bAllHandlesVisible = true;
-	UActorComponent* SingleVisibleHandle = nullptr;
-	TArray<UActorComponent*> HoveringOverHandles;
-	RefreshTransformGizmo( bNewObjectsSelected, bAllHandlesVisible, SingleVisibleHandle, HoveringOverHandles );
+	RefreshTransformGizmo( bNewObjectsSelected );
 }
 
 void UViewportWorldInteraction::SetDefaultOptionalViewportClient(const TSharedPtr<class FEditorViewportClient>& InEditorViewportClient)
@@ -437,8 +550,8 @@ bool UViewportWorldInteraction::InputKey( FEditorViewportClient* InViewportClien
 
 	if( bActive )
 	{
-		check( InViewportClient != nullptr );
-		OnKeyInputEvent.Broadcast( InViewportClient, Key, Event, bWasHandled );
+		check(InViewportClient != nullptr);
+		OnKeyInputEvent.Broadcast(InViewportClient, Key, Event, bWasHandled);
 
 		if( !bWasHandled || !bActive )
 		{
@@ -471,7 +584,7 @@ bool UViewportWorldInteraction::InputKey( FEditorViewportClient* InViewportClien
 
 bool UViewportWorldInteraction::PreprocessedInputKey( const FKey Key, const EInputEvent Event )
 {
-	if( DefaultOptionalViewportClient != nullptr && bUseInputProcessor == true )
+	if( DefaultOptionalViewportClient != nullptr && bUseInputPreprocessor == true )
 	{
 		return InputKey( DefaultOptionalViewportClient, DefaultOptionalViewportClient->Viewport, Key, Event );
 	}
@@ -508,7 +621,7 @@ bool UViewportWorldInteraction::InputAxis( FEditorViewportClient* InViewportClie
 
 bool UViewportWorldInteraction::PreprocessedInputAxis( const int32 ControllerId, const FKey Key, const float Delta, const float DeltaTime )
 {
-	if( DefaultOptionalViewportClient != nullptr && bUseInputProcessor == true )
+	if( DefaultOptionalViewportClient != nullptr && bUseInputPreprocessor == true )
 	{
 		return InputAxis( DefaultOptionalViewportClient, DefaultOptionalViewportClient->Viewport, ControllerId, Key, Delta, DeltaTime );
 	}
@@ -578,22 +691,6 @@ float UViewportWorldInteraction::GetWorldScaleFactor() const
 	return GetWorld()->GetWorldSettings()->WorldToMeters / 100.0f;
 }
 
-UWorld* UViewportWorldInteraction::GetWorld() const
-{
-	UWorld* ResultWorld = nullptr;
-	
-	if( UEditorWorldExtension::GetWorld() != nullptr )
-	{
-		ResultWorld = UEditorWorldExtension::GetWorld();
-	}
-	else if ( DefaultOptionalViewportClient != nullptr )
-	{
-		ResultWorld = DefaultOptionalViewportClient->GetWorld();
-	}
-
-	return ResultWorld;
-}
-
 FEditorViewportClient* UViewportWorldInteraction::GetDefaultOptionalViewportClient() const
 {
 	return DefaultOptionalViewportClient;
@@ -601,36 +698,39 @@ FEditorViewportClient* UViewportWorldInteraction::GetDefaultOptionalViewportClie
 
 void UViewportWorldInteraction::Undo()
 {
-	GUnrealEd->Exec( GetWorld(), TEXT( "TRANSACTION UNDO" ) );
+	PlaySound(AssetContainer->UndoSound, TransformGizmoActor->GetActorLocation());
+	bPlayNextRefreshTransformGizmoSound = false;
+	ExecCommand(FString("TRANSACTION UNDO"));
 }
 
 void UViewportWorldInteraction::Redo()
 {
-	GUnrealEd->Exec( GetWorld(), TEXT( "TRANSACTION REDO" ) );
+	PlaySound(AssetContainer->RedoSound, TransformGizmoActor->GetActorLocation());
+	bPlayNextRefreshTransformGizmoSound = false;
+	ExecCommand(FString("TRANSACTION REDO"));
 }
 
 void UViewportWorldInteraction::DeleteSelectedObjects()
 {
-	GUnrealEd->Exec(GetWorld(), TEXT("DELETE"));
+	ExecCommand(FString("DELETE"));
 }
 
 void UViewportWorldInteraction::Copy()
 {
 	// @todo vreditor: Needs CanExecute()  (see LevelEditorActions.cpp)
-	GUnrealEd->Exec( GetWorld(), TEXT( "EDIT COPY" ) );
+	ExecCommand(FString("EDIT COPY"));
 }
 
 void UViewportWorldInteraction::Paste()
 {
 	// @todo vreditor: Needs CanExecute()  (see LevelEditorActions.cpp)
 	// @todo vreditor: Needs "paste here" style pasting (TO=HERE), but with ray
-	GUnrealEd->Exec( GetWorld(), TEXT( "EDIT PASTE" ) );
+	ExecCommand(FString("EDIT PASTE"));
 }
 
 void UViewportWorldInteraction::Duplicate()
 {
-	// @todo vreditor: Needs CanExecute()  (see LevelEditorActions.cpp)
-	GUnrealEd->Exec( GetWorld(), TEXT( "DUPLICATE" ) );
+	ExecCommand(FString("DUPLICATE"));
 }
 
 void UViewportWorldInteraction::Deselect()
@@ -645,12 +745,9 @@ void UViewportWorldInteraction::HoverTick( const float DeltaTime )
 		bool bWasHandled = false;
 		FViewportInteractorData& InteractorData = Interactor->GetInteractorData();
 		
-		Interactor->ResetHoverState( DeltaTime );
+		Interactor->ResetHoverState();
 		
-		//Update the state of the interactor here
-		FHitResult HitHoverResult;
-		Interactor->UpdateHoverResult( HitHoverResult );
-		FVector HoverImpactPoint = HitHoverResult.ImpactPoint;
+		FHitResult HitHoverResult = Interactor->GetHitResultFromLaserPointer();
 
 		const bool bIsHoveringOverTransformGizmo =
 			HitHoverResult.Actor.IsValid() &&
@@ -659,26 +756,18 @@ void UViewportWorldInteraction::HoverTick( const float DeltaTime )
 		// Prefer transform gizmo hover over everything else
 		if( !bIsHoveringOverTransformGizmo )
 		{
+			FVector HoverImpactPoint = FVector::ZeroVector;
 			OnHoverUpdateEvent.Broadcast( Interactor, /* In/Out */ HoverImpactPoint, /* In/Out */ bWasHandled );
 			if( bWasHandled )
 			{
-				InteractorData.bIsHovering = true;
-				InteractorData.HoverLocation = HoverImpactPoint;
+				Interactor->SetHoverLocation(HoverImpactPoint);
 			}
 		}
 		
-		if ( bWasHandled )
-		{
-			InteractorData.bIsHovering = true;
-			InteractorData.HoverLocation = HoverImpactPoint;
-		}
-		else
+		if ( !bWasHandled )
 		{
 			UActorComponent* NewHoveredActorComponent = nullptr;
 
-			InteractorData.HoveringOverTransformGizmoComponent = nullptr;
-			InteractorData.HoverLocation = FVector::ZeroVector;
-
 			if ( HitHoverResult.Actor.IsValid() )
 			{
 				USceneComponent* HoveredActorComponent = HitHoverResult.GetComponent();
@@ -688,7 +777,9 @@ void UViewportWorldInteraction::HoverTick( const float DeltaTime )
 				{
 					HoveredObjects.Add( FViewportHoverTarget( Actor ) );
 
-					InteractorData.HoverLocation = HitHoverResult.ImpactPoint;
+					Interactor->SetHoverLocation(HitHoverResult.ImpactPoint);
+
+					bWasHandled = true;
 
 					if ( Actor == TransformGizmoActor )
 					{
@@ -704,7 +795,7 @@ void UViewportWorldInteraction::HoverTick( const float DeltaTime )
 							NewHoveredActorComponent = HoveredActorComponent;
 
 							// Check if the current hovered component of the interactor is different from the hitresult component
-							if ( NewHoveredActorComponent != InteractorData.HoveredActorComponent && !IsOtherInteractorHoveringOverComponent( Interactor, NewHoveredActorComponent ) )
+							if ( NewHoveredActorComponent != InteractorData.LastHoveredActorComponent && !IsOtherInteractorHoveringOverComponent( Interactor, NewHoveredActorComponent ) )
 							{
 								Interactable->OnHoverEnter( Interactor, HitHoverResult );
 							}
@@ -712,20 +803,15 @@ void UViewportWorldInteraction::HoverTick( const float DeltaTime )
 							Interactable->OnHover( Interactor );
 						}
 					}
-					InteractorData.bIsHovering = true;
-				}
-				else if ( !HoveredActorComponent )
-				{
-					InteractorData.bIsHovering = false;
 				}
 			}
 
 			// Leave hovered interactable
 			//@todo ViewportInteraction: This does not take into account when the other interactors are already hovering over this interactable
-			if ( InteractorData.HoveredActorComponent != nullptr && ( InteractorData.HoveredActorComponent != NewHoveredActorComponent || NewHoveredActorComponent == nullptr ) 
+			if ( InteractorData.LastHoveredActorComponent != nullptr && ( InteractorData.LastHoveredActorComponent != NewHoveredActorComponent || NewHoveredActorComponent == nullptr ) 
 				&& !IsOtherInteractorHoveringOverComponent( Interactor, NewHoveredActorComponent ) )
 			{
-				AActor* HoveredActor = InteractorData.HoveredActorComponent->GetOwner();
+				AActor* HoveredActor = InteractorData.LastHoveredActorComponent->GetOwner();
 				if ( HoveredActor )
 				{
 					IViewportInteractableInterface* Interactable = Cast<IViewportInteractableInterface>( HoveredActor );
@@ -737,7 +823,12 @@ void UViewportWorldInteraction::HoverTick( const float DeltaTime )
 			}
 
 			//Update the hovered actor component with the new component
-			InteractorData.HoveredActorComponent = NewHoveredActorComponent;
+			InteractorData.LastHoveredActorComponent = NewHoveredActorComponent;
+		}
+
+		if (InteractorData.LastHoveredActorComponent == nullptr && HitHoverResult.GetComponent() != nullptr)
+		{
+			InteractorData.LastHoveredActorComponent = HitHoverResult.GetComponent();
 		}
 	}
 }
@@ -765,7 +856,7 @@ void UViewportWorldInteraction::InteractionTick( const float DeltaTime )
 	const float WorldScaleFactor = GetWorldScaleFactor();
 
 	// Move selected actors
-	UViewportInteractor* DraggingTransformablesWithInteractor = nullptr;
+	UViewportInteractor* DraggingWithInteractor = nullptr;
 	UViewportInteractor* AssistingDragWithInteractor = nullptr;
 
 	for ( UViewportInteractor* Interactor : Interactors )
@@ -777,12 +868,15 @@ void UViewportWorldInteraction::InteractionTick( const float DeltaTime )
 			InteractorData.DraggingMode == EViewportInteractionDraggingMode::TransformablesFreely ||
 			InteractorData.DraggingMode == EViewportInteractionDraggingMode::TransformablesAtLaserImpact )
 		{
-			check( DraggingTransformablesWithInteractor == nullptr );	// Only support dragging one thing at a time right now!
-			DraggingTransformablesWithInteractor = Interactor;
+			check( DraggingWithInteractor == nullptr );	// Only support dragging one thing at a time right now!
+			DraggingWithInteractor = Interactor;
 
 			if ( InteractorData.DraggingMode != EViewportInteractionDraggingMode::TransformablesAtLaserImpact )
 			{
-				bCanSlideRayLength = true;
+				if( !InteractorData.bDraggingWithGrabberSphere )
+				{
+					bCanSlideRayLength = true;
+				}
 			}
 		}
 
@@ -791,18 +885,19 @@ void UViewportWorldInteraction::InteractionTick( const float DeltaTime )
 			check( AssistingDragWithInteractor == nullptr );	// Only support assisting one thing at a time right now!
 			AssistingDragWithInteractor = Interactor;
 
-			bCanSlideRayLength = true;
+			if( !InteractorData.bDraggingWithGrabberSphere )
+			{
+				bCanSlideRayLength = true;
+			}
 		}
 
 		if ( bCanSlideRayLength )
 		{
-			Interactor->CalculateDragRay();
+			Interactor->CalculateDragRay( InteractorData.DragRayLength, InteractorData.DragRayLengthVelocity );
 		}
 	}
 
-	const bool bIsStillDraggingTransformables = DraggingTransformablesWithInteractor != nullptr;
-
-	bool bAllInteractorsThatWereDraggingTransformablesHaveZeroVelocity = true;	
+	UViewportInteractor* InertiaFromInteractor = nullptr;
 	for ( UViewportInteractor* Interactor : Interactors )
 	{
 		FViewportInteractorData& InteractorData = Interactor->GetInteractorData();
@@ -810,7 +905,7 @@ void UViewportWorldInteraction::InteractionTick( const float DeltaTime )
 		if( InteractorData.DraggingMode == EViewportInteractionDraggingMode::TransformablesWithGizmo ||
 			InteractorData.DraggingMode == EViewportInteractionDraggingMode::TransformablesFreely ||
 			InteractorData.DraggingMode == EViewportInteractionDraggingMode::TransformablesAtLaserImpact ||
-			InteractorData.DraggingMode == EViewportInteractionDraggingMode::World )
+			( InteractorData.DraggingMode == EViewportInteractionDraggingMode::World && bAllowWorldMovement ) )
 		{
 			// Are we dragging with two interactors ?
 			UViewportInteractor* OtherInteractor = nullptr;
@@ -823,7 +918,7 @@ void UViewportWorldInteraction::InteractionTick( const float DeltaTime )
 
 			FVector DraggedTo = InteractorData.Transform.GetLocation();
 			FVector DragDelta = DraggedTo - InteractorData.LastTransform.GetLocation();
-			FVector DragDeltaFromStart = DraggedTo - InteractorData.LaserPointerImpactAtDragStart;
+			FVector DragDeltaFromStart = DraggedTo - InteractorData.ImpactLocationAtDragStart;
 
 			FVector OtherHandDraggedTo = FVector::ZeroVector;
 			FVector OtherHandDragDelta = FVector::ZeroVector;
@@ -832,13 +927,13 @@ void UViewportWorldInteraction::InteractionTick( const float DeltaTime )
 			{
 				OtherHandDraggedTo = OtherInteractor->GetInteractorData().Transform.GetLocation();
 				OtherHandDragDelta = OtherHandDraggedTo - OtherInteractor->GetInteractorData().LastTransform.GetLocation();
-				OtherHandDragDeltaFromStart = DraggedTo - OtherInteractor->GetInteractorData().LaserPointerImpactAtDragStart;
+				OtherHandDragDeltaFromStart = DraggedTo - OtherInteractor->GetInteractorData().ImpactLocationAtDragStart;
 			}
 			else if( OtherInteractorThatWasAssistingDrag != nullptr )
 			{
 				OtherHandDragDelta = OtherInteractorThatWasAssistingDrag->GetInteractorData().DragTranslationVelocity;
 				OtherHandDraggedTo = OtherInteractorThatWasAssistingDrag->GetInteractorData().LastDragToLocation + OtherHandDragDelta;
-				OtherHandDragDeltaFromStart = OtherHandDraggedTo - OtherInteractorThatWasAssistingDrag->GetInteractorData().LaserPointerImpactAtDragStart;
+				OtherHandDragDeltaFromStart = OtherHandDraggedTo - OtherInteractorThatWasAssistingDrag->GetInteractorData().ImpactLocationAtDragStart;
 			}
 
 
@@ -851,7 +946,36 @@ void UViewportWorldInteraction::InteractionTick( const float DeltaTime )
 				InteractorData.DraggingMode == EViewportInteractionDraggingMode::TransformablesAtLaserImpact )
 			{
 				// Move objects using the laser pointer (in world space)
-				if( Interactor->GetLaserPointer( /* Out */ LaserPointerStart, /* Out */ LaserPointerEnd ) )
+				if( InteractorData.bDraggingWithGrabberSphere )
+				{
+					FSphere GrabberSphere;
+					if( Interactor->GetGrabberSphere( /* Out */ GrabberSphere ) )
+					{
+						// @todo grabber:  Fill in	LaserPointerStart	LaserPointerEnd		LaserPointerDirection   (not used for anything when in TransformablesFreely mode)
+
+						const FVector RotatedOffsetFromSphereCenterAtDragStart = InteractorData.ImpactLocationAtDragStart - InteractorData.GrabberSphereLocationAtDragStart;
+						const FVector UnrotatedOffsetFromSphereCenterAtDragStart = InteractorData.InteractorRotationAtDragStart.UnrotateVector( RotatedOffsetFromSphereCenterAtDragStart );
+						DraggedTo = GrabberSphere.Center + InteractorData.Transform.GetRotation().RotateVector( UnrotatedOffsetFromSphereCenterAtDragStart );
+
+						const FVector WorldSpaceDragDelta = DraggedTo - InteractorData.LastDragToLocation;
+						DragDelta = WorldSpaceDragDelta;
+						InteractorData.DragTranslationVelocity = WorldSpaceDragDelta;
+
+						const FVector WorldSpaceDeltaFromStart = DraggedTo - InteractorData.ImpactLocationAtDragStart;
+						DragDeltaFromStart = WorldSpaceDeltaFromStart;
+
+						InteractorData.LastDragToLocation = DraggedTo;
+
+						// Update hover location (we only do this when dragging using the laser pointer)
+						Interactor->SetHoverLocation(DraggedTo);
+					}
+					else
+					{
+						// We lost our grabber sphere, so cancel the drag
+						StopDragging( Interactor );
+					}
+				}
+				else if( Interactor->GetLaserPointer( /* Out */ LaserPointerStart, /* Out */ LaserPointerEnd ) )
 				{
 					LaserPointerDirection = ( LaserPointerEnd - LaserPointerStart ).GetSafeNormal();
 
@@ -860,37 +984,40 @@ void UViewportWorldInteraction::InteractionTick( const float DeltaTime )
 						// Check to see if the laser pointer is over something we can drop on
 						FVector HitLocation = FVector::ZeroVector;
 						bool bHitSomething = FindPlacementPointUnderLaser( Interactor, /* Out */ HitLocation );
+						if( !bHitSomething )
+						{
+							HitLocation = LaserPointerStart + LaserPointerDirection * InteractorData.DragRayLength;
+						}
 
-						if( bHitSomething )
+						if( InteractorData.bIsFirstDragUpdate || !bFreezePlacementWhileInterpolatingTransformables || !bIsInterpolatingTransformablesFromSnapshotTransform )
 						{
+							InteractorData.DragRayLength = ( LaserPointerStart - HitLocation ).Size();
+							DraggedTo = HitLocation;
+
 							// If the object moved reasonably far between frames, it might be because the angle we were aligning
 							// the object with during placement changed radically.  To avoid it popping, we smoothly interpolate
 							// it's position over a very short timespan
 							if( !bIsInterpolatingTransformablesFromSnapshotTransform )	// Let the last animation finish first
 							{
-								const float ScaledDragDistance = DragDelta.Size() * WorldScaleFactor;
+								const FVector WorldSpaceDragDelta = InteractorData.bIsFirstDragUpdate ? FVector::ZeroVector : ( DraggedTo - InteractorData.LastDragToLocation );
+								const float ScaledDragDistance = WorldSpaceDragDelta.Size() * WorldScaleFactor;
 								if( ScaledDragDistance >= VI::DragAtLaserImpactInterpolationThreshold->GetFloat() )
 								{
 									bIsInterpolatingTransformablesFromSnapshotTransform = true;
+									bFreezePlacementWhileInterpolatingTransformables = true;
 									TransformablesInterpolationStartTime = CurrentTime;
 									TransformablesInterpolationDuration = VI::DragAtLaserImpactInterpolationDuration->GetFloat();
 
-									// Snapshot everything
-									for( TUniquePtr<FViewportTransformable>& TransformablePtr : Transformables )
-									{
-										FViewportTransformable& Transformable = *TransformablePtr;
-										Transformable.InterpolationSnapshotTransform = Transformable.LastTransform;
-									}
+									// Snapshot time!
+									InteractorData.GizmoInterpolationSnapshotTransform = InteractorData.GizmoLastTransform;
 								}
 							}
-
-							InteractorData.DragRayLength = ( LaserPointerStart - HitLocation ).Size();
-
-							DraggedTo = HitLocation;
 						}
 						else
 						{
-							DraggedTo = LaserPointerStart + LaserPointerDirection * InteractorData.DragRayLength;
+							// Keep interpolating toward the previous placement location
+							check( bFreezePlacementWhileInterpolatingTransformables );
+							DraggedTo = InteractorData.LastDragToLocation;
 						}
 					}
 					else
@@ -902,15 +1029,13 @@ void UViewportWorldInteraction::InteractionTick( const float DeltaTime )
 					DragDelta = WorldSpaceDragDelta;
 					InteractorData.DragTranslationVelocity = WorldSpaceDragDelta;
 
-					const FVector WorldSpaceDeltaFromStart = DraggedTo - InteractorData.LaserPointerImpactAtDragStart;
+					const FVector WorldSpaceDeltaFromStart = DraggedTo - InteractorData.ImpactLocationAtDragStart;
 					DragDeltaFromStart = WorldSpaceDeltaFromStart;
 
-					InteractorData.LastLaserPointerStart = LaserPointerStart;
 					InteractorData.LastDragToLocation = DraggedTo;
 
 					// Update hover location (we only do this when dragging using the laser pointer)
-					InteractorData.bIsHovering = true;
-					InteractorData.HoverLocation = FMath::ClosestPointOnLine( LaserPointerStart, LaserPointerEnd, DraggedTo );
+					Interactor->SetHoverLocation(FMath::ClosestPointOnLine(LaserPointerStart, LaserPointerEnd, DraggedTo));
 				}
 				else
 				{
@@ -918,7 +1043,7 @@ void UViewportWorldInteraction::InteractionTick( const float DeltaTime )
 					StopDragging( Interactor );
 				}
 
-				if( OtherInteractor != nullptr )
+				if( OtherInteractor != nullptr )	// @todo grabber: Second hand support (should be doable)
 				{
 					FVector OtherHandLaserPointerStart, OtherHandLaserPointerEnd;
 					if( OtherInteractor->GetLaserPointer( /* Out */ OtherHandLaserPointerStart, /* Out */ OtherHandLaserPointerEnd ) )
@@ -932,15 +1057,13 @@ void UViewportWorldInteraction::InteractionTick( const float DeltaTime )
 						OtherHandDragDelta = OtherHandWorldSpaceDragDelta;
 						OtherInteractorData.DragTranslationVelocity = OtherHandWorldSpaceDragDelta;
 
-						const FVector OtherHandWorldSpaceDeltaFromStart = OtherHandDraggedTo - OtherInteractorData.LaserPointerImpactAtDragStart;
+						const FVector OtherHandWorldSpaceDeltaFromStart = OtherHandDraggedTo - OtherInteractorData.ImpactLocationAtDragStart;
 						OtherHandDragDeltaFromStart = OtherHandWorldSpaceDeltaFromStart;
 
-						OtherInteractorData.LastLaserPointerStart = OtherHandLaserPointerStart;
 						OtherInteractorData.LastDragToLocation = OtherHandDraggedTo;
 
 						// Only hover if we're using the laser pointer
-						OtherInteractorData.bIsHovering = true;
-						OtherInteractorData.HoverLocation = OtherHandDraggedTo;
+						OtherInteractor->SetHoverLocation(OtherHandDraggedTo);
 					}
 					else
 					{
@@ -1029,6 +1152,10 @@ void UViewportWorldInteraction::InteractionTick( const float DeltaTime )
 				Interactor->GetLaserPointerMaxLength(),
 				bIsLaserPointerValid, 
 				InteractorData.GizmoStartTransform, 
+				InteractorData.GizmoLastTransform,
+				InteractorData.GizmoTargetTransform,
+				InteractorData.GizmoUnsnappedTargetTransform,
+				InteractorData.GizmoInterpolationSnapshotTransform,
 				InteractorData.GizmoStartLocalBounds,
 				InteractorData.DraggingTransformGizmoComponent.Get(),
 				/* In/Out */ InteractorData.GizmoSpaceFirstDragUpdateOffsetAlongAxis,
@@ -1116,12 +1243,12 @@ void UViewportWorldInteraction::InteractionTick( const float DeltaTime )
 					InteractorData.LastDraggingMode == EViewportInteractionDraggingMode::TransformablesFreely ||
 					InteractorData.LastDraggingMode == EViewportInteractionDraggingMode::TransformablesAtLaserImpact )
 				{
-					bAllInteractorsThatWereDraggingTransformablesHaveZeroVelocity = false;
+					InertiaFromInteractor = Interactor;
 				}
 
 				const FVector DragDelta = InteractorData.DragTranslationVelocity;
 				const FVector DraggedTo = InteractorData.LastDragToLocation + DragDelta;
-				const FVector DragDeltaFromStart = DraggedTo - InteractorData.LaserPointerImpactAtDragStart;
+				const FVector DragDeltaFromStart = DraggedTo - InteractorData.ImpactLocationAtDragStart;
 
 				UViewportInteractor* OtherInteractorThatWasAssistingDrag = GetOtherInteractorIntertiaContribute( Interactor );
 
@@ -1135,7 +1262,7 @@ void UViewportWorldInteraction::InteractionTick( const float DeltaTime )
 					FViewportInteractorData& OtherInteractorThatWasAssistingDragData = OtherInteractorThatWasAssistingDrag->GetInteractorData();
 					OtherHandDragDelta = OtherInteractorThatWasAssistingDragData.DragTranslationVelocity;
 					OtherHandDraggedTo = OtherInteractorThatWasAssistingDragData.LastDragToLocation + OtherHandDragDelta;
-					OtherHandDragDeltaFromStart = OtherHandDraggedTo - OtherInteractorThatWasAssistingDragData.LaserPointerImpactAtDragStart;
+					OtherHandDragDeltaFromStart = OtherHandDraggedTo - OtherInteractorThatWasAssistingDragData.ImpactLocationAtDragStart;
 				}
 
 				const bool bIsLaserPointerValid = false;	// Laser pointer has moved on to other things
@@ -1158,7 +1285,11 @@ void UViewportWorldInteraction::InteractionTick( const float DeltaTime )
 					0.0f,		// No valid laser pointer during inertia
 					bIsLaserPointerValid, 
 					InteractorData.GizmoStartTransform, 
-					InteractorData.GizmoStartLocalBounds, 
+					InteractorData.GizmoLastTransform,
+					InteractorData.GizmoTargetTransform,
+					InteractorData.GizmoUnsnappedTargetTransform,
+					InteractorData.GizmoInterpolationSnapshotTransform,
+					InteractorData.GizmoStartLocalBounds,
 					InteractorData.DraggingTransformGizmoComponent.Get(),
 					/* In/Out */ InteractorData.GizmoSpaceFirstDragUpdateOffsetAlongAxis,
 					/* In/Out */ InteractorData.GizmoSpaceDragDeltaFromStartOffset,
@@ -1183,40 +1314,18 @@ void UViewportWorldInteraction::InteractionTick( const float DeltaTime )
 		}
 	}
 
-	// Refresh the transform gizmo every frame, just in case actors were moved by some external
-	// influence.  Also, some features of the transform gizmo respond to camera position (like the
-	// measurement text labels), so it just makes sense to keep it up to date.
-	{
-		const bool bNewObjectsSelected = false;
-		const bool bAllHandlesVisible = ( DraggingTransformablesWithInteractor == nullptr );
-		UActorComponent* SingleVisibleHandle = ( DraggingTransformablesWithInteractor != nullptr ) ? DraggingTransformablesWithInteractor->GetInteractorData().DraggingTransformGizmoComponent.Get() : nullptr;
-
-		static TArray< UActorComponent* > HoveringOverHandles;
-		HoveringOverHandles.Reset();
-		for( UViewportInteractor* Interactor : Interactors )
-		{
-			FViewportInteractorData& InteractorData = Interactor->GetInteractorData();
-			UActorComponent* HoveringOverHandle = InteractorData.HoveringOverTransformGizmoComponent.Get();
-			if( HoveringOverHandle != nullptr )
-			{
-				HoveringOverHandles.Add( HoveringOverHandle );
-			}
-		}
-
-		RefreshTransformGizmo( bNewObjectsSelected, bAllHandlesVisible, SingleVisibleHandle, HoveringOverHandles );
-	}
-
 
 	// Update transformables
 	const bool bSmoothSnappingEnabled = IsSmoothSnappingEnabled();
 	if ( bAreTransformablesMoving && ( bSmoothSnappingEnabled || bIsInterpolatingTransformablesFromSnapshotTransform ) )
 	{
 		const float SmoothSnapSpeed = VI::SmoothSnapSpeed->GetFloat();
-		const bool bUseElasticSnapping = bSmoothSnappingEnabled && VI::ElasticSnap->GetInt() > 0 && bIsStillDraggingTransformables;	// Only while we're still dragging stuff!
+		const bool bUseElasticSnapping = bSmoothSnappingEnabled && 
+										VI::ElasticSnap->GetInt() > 0 && 
+										DraggingWithInteractor != nullptr &&	// Only while we're still dragging stuff!
+										VI::EnableGuides->GetInt() == 0;	// Not while using actor align/snap
 		const float ElasticSnapStrength = VI::ElasticSnapStrength->GetFloat();
 
-		bool bIsStillSmoothSnapping = false;
-
 		float InterpProgress = 1.0f;
 		const bool bWasInterpolationNeeded = bIsInterpolatingTransformablesFromSnapshotTransform;
 		bool bDidInterpolationFinishThisUpdate = false;
@@ -1227,60 +1336,69 @@ void UViewportWorldInteraction::InteractionTick( const float DeltaTime )
 			{
 				// Finished interpolating
 				bIsInterpolatingTransformablesFromSnapshotTransform = false;
+				bFreezePlacementWhileInterpolatingTransformables = false;
 				bDidInterpolationFinishThisUpdate = true;
 			}
 		}
 
-		for ( TUniquePtr<FViewportTransformable>& TransformablePtr : Transformables )
+		bool bIsStillSmoothSnapping = false;
+
+		UViewportInteractor* TransformingInteractor = DraggingWithInteractor != nullptr ? DraggingWithInteractor : InertiaFromInteractor;
+		if( TransformingInteractor != nullptr )
 		{
-			FViewportTransformable& Transformable = *TransformablePtr;
+			FViewportInteractorData& InteractorData = TransformingInteractor->GetInteractorData();
 
-			FTransform ActualTargetTransform = Transformable.TargetTransform;
+			FTransform ActualGizmoTargetTransform = InteractorData.GizmoTargetTransform;
 
 			// If 'elastic snapping' is turned on, we'll have the object 'reach' toward its unsnapped position, so
 			// that the user always has visual feedback when they are dragging something around, even if they
 			// haven't dragged far enough to overcome the snap threshold.
-			if ( bUseElasticSnapping )
+			if( bUseElasticSnapping )
 			{
-				ActualTargetTransform.BlendWith( Transformable.UnsnappedTargetTransform, ElasticSnapStrength );
+				ActualGizmoTargetTransform.BlendWith( InteractorData.GizmoUnsnappedTargetTransform, ElasticSnapStrength );
 			}
 
-			FTransform InterpolatedTransform = ActualTargetTransform;
-			if ( !ActualTargetTransform.Equals( Transformable.LastTransform, 0.0f ) )
+			FTransform InterpolatedGizmoTransform = ActualGizmoTargetTransform;
+			if( !ActualGizmoTargetTransform.Equals( InteractorData.GizmoLastTransform, 0.0f ) )
 			{
 				// If we're really close, just snap it.
-				if ( ActualTargetTransform.Equals( Transformable.LastTransform, KINDA_SMALL_NUMBER ) )
+				if( ActualGizmoTargetTransform.Equals( InteractorData.GizmoLastTransform, KINDA_SMALL_NUMBER ) )
 				{
-					InterpolatedTransform = Transformable.LastTransform = Transformable.UnsnappedTargetTransform = ActualTargetTransform;
+					InterpolatedGizmoTransform = InteractorData.GizmoLastTransform = ActualGizmoTargetTransform;
 				}
 				else
 				{
 					bIsStillSmoothSnapping = true;
 
-					InterpolatedTransform.Blend(
-						Transformable.LastTransform,
-						ActualTargetTransform,
+					InterpolatedGizmoTransform.Blend(
+						InteractorData.GizmoLastTransform,
+						ActualGizmoTargetTransform,
 						FMath::Min( 1.0f, SmoothSnapSpeed * FMath::Min( DeltaTime, 1.0f / 30.0f ) ) );
 				}
-				Transformable.LastTransform = InterpolatedTransform;
+				InteractorData.GizmoLastTransform = InterpolatedGizmoTransform;
 			}
 
-			if ( bIsInterpolatingTransformablesFromSnapshotTransform )
+			if( bIsInterpolatingTransformablesFromSnapshotTransform )
 			{
-				InterpolatedTransform.BlendWith( Transformable.InterpolationSnapshotTransform, 1.0f - InterpProgress );
+				InterpolatedGizmoTransform.BlendWith( InteractorData.GizmoInterpolationSnapshotTransform, 1.0f - InterpProgress );
 			}
 
-			// Update the transform!
-			// NOTE: We never sweep while smooth-snapping as the object will never end up where we want it to
-			// due to friction with the ground and other objects.
-			const bool bSweep = false;
-			Transformable.ApplyTransform( InterpolatedTransform, bSweep );
+			for ( TUniquePtr<FViewportTransformable>& TransformablePtr : Transformables )
+			{
+				FViewportTransformable& Transformable = *TransformablePtr;
+
+				// Update the transform!
+				// NOTE: We never sweep while smooth-snapping as the object will never end up where we want it to
+				// due to friction with the ground and other objects.
+				const bool bSweep = false;
+				Transformable.ApplyTransform( Transformable.StartTransform * InteractorData.GizmoStartTransform.Inverse() * InterpolatedGizmoTransform, bSweep );
+			}
 		}
 
 		// Did we finish interpolating in this update?  If so, we need to finalize things and notify everyone
 		if( bAreTransformablesMoving &&
-			!bIsStillDraggingTransformables &&
-			bAllInteractorsThatWereDraggingTransformablesHaveZeroVelocity &&
+			DraggingWithInteractor == nullptr &&
+			InertiaFromInteractor == nullptr &&
 			( !bWasInterpolationNeeded || bDidInterpolationFinishThisUpdate ) && 
 			!bIsStillSmoothSnapping )
 		{
@@ -1292,13 +1410,23 @@ void UViewportWorldInteraction::InteractionTick( const float DeltaTime )
 		// Neither smooth snapping or interpolation is enabled right now, but the transformables could have some velocity
 		// applied.  We'll check to see whether they've come to a rest, and if so finalize their positions.
 		if( bAreTransformablesMoving && 
-			!bIsStillDraggingTransformables &&
-			bAllInteractorsThatWereDraggingTransformablesHaveZeroVelocity )
+			DraggingWithInteractor == nullptr &&
+			InertiaFromInteractor == nullptr )
 		{
 			FinishedMovingTransformables();
 		}
 	}
 
+
+	// Refresh the transform gizmo every frame, just in case actors were moved by some external
+	// influence.  Also, some features of the transform gizmo respond to camera position (like the
+	// measurement text labels), so it just makes sense to keep it up to date.
+	{
+		const bool bNewObjectsSelected = false;
+		RefreshTransformGizmo( bNewObjectsSelected );
+	}
+
+
 	LastWorldToMetersScale = WorldToMetersScale;
 } 
 
@@ -1359,6 +1487,10 @@ void UViewportWorldInteraction::UpdateDragging(
 	const float LaserPointerMaxLength,
 	const bool bIsLaserPointerValid,
 	const FTransform& GizmoStartTransform,
+	FTransform& GizmoLastTransform,
+	FTransform& GizmoTargetTransform,
+	FTransform& GizmoUnsnappedTargetTransform,
+	const FTransform& GizmoInterpolationSnapshotTransform,
 	const FBox& GizmoStartLocalBounds,
 	const USceneComponent* const DraggingTransformGizmoComponent,
 	FVector& GizmoSpaceFirstDragUpdateOffsetAlongAxis,
@@ -1370,7 +1502,7 @@ void UViewportWorldInteraction::UpdateDragging(
 	//		-> When dragging objects around, everything is in world space unless otherwise specified in the variable name.
 	//		-> When dragging the world around, everything is in room space unless otherwise specified.
 
-	bool bMovedAnyTransformables = false;
+	bool bMovedTransformGizmo = false;
 
 	// This will be set to true if we want to set the physics velocity on the dragged objects based on how far
 	// we dragged this frame.  Used for pushing objects around in Simulate mode.  If this is set to false, we'll
@@ -1384,12 +1516,12 @@ void UViewportWorldInteraction::UpdateDragging(
 	const bool bLocalSpaceSnapping =
 		GetTransformGizmoCoordinateSpace() == COORD_Local &&
 		DraggingMode != EViewportInteractionDraggingMode::TransformablesAtLaserImpact;
-	const FVector SnapGridBase = ( DraggingMode == EViewportInteractionDraggingMode::TransformablesAtLaserImpact ||  bLocalSpaceSnapping ) ? FVector::ZeroVector : GizmoStartTransform.GetLocation(); //@todo ViewportInteraction
+	const FVector SnapGridBase = ( DraggingMode == EViewportInteractionDraggingMode::TransformablesAtLaserImpact ||  bLocalSpaceSnapping ) ? FVector::ZeroVector : GizmoStartTransform.GetLocation();
 
 
 	// Okay, time to move stuff!  We'll do this part differently depending on whether we're dragging actual actors around
 	// or we're moving the camera (aka. dragging the world)
-	if ( DraggingMode == EViewportInteractionDraggingMode::TransformablesWithGizmo ||
+	if( DraggingMode == EViewportInteractionDraggingMode::TransformablesWithGizmo ||
 		DraggingMode == EViewportInteractionDraggingMode::TransformablesAtLaserImpact )
 	{
 		FVector OutClosestPointOnLaser;
@@ -1407,26 +1539,22 @@ void UViewportWorldInteraction::UpdateDragging(
 			/* In/Out */ DragDeltaFromStartOffset,
 			/* Out */ OutClosestPointOnLaser);
 
+		FVector OriginalConstrainedDragDelta = ConstrainedDragDeltaFromStart;
 		ConstrainedDragDeltaFromStart = GizmoStartTransform.GetLocation() + ConstrainedDragDeltaFromStart;
 
 		// Set out put for hover point
 		OutUnsnappedDraggedTo = GizmoStartTransform.GetLocation() + OutClosestPointOnLaser;
 
+
 		// Grid snap!
-		FVector SnappedDraggedTo = ConstrainedDragDeltaFromStart;
-		if ( FSnappingUtils::IsSnapToGridEnabled() )
+		FVector DesiredGizmoLocation = ConstrainedDragDeltaFromStart;
+		FVector SnappedDraggedTo = DesiredGizmoLocation;
+		if ((InteractionType == ETransformGizmoInteractionType::Translate ||
+			InteractionType == ETransformGizmoInteractionType::TranslateOnPlane) || 
+			DraggingMode == EViewportInteractionDraggingMode::TransformablesAtLaserImpact)
 		{
-			// Snap in local space, if we need to
-			if ( bLocalSpaceSnapping )
-			{
-				SnappedDraggedTo = GizmoStartTransform.InverseTransformPositionNoScale( SnappedDraggedTo );
-				FSnappingUtils::SnapPointToGrid( SnappedDraggedTo, SnapGridBase );
-				SnappedDraggedTo = GizmoStartTransform.TransformPositionNoScale( SnappedDraggedTo );
-			}
-			else
-			{
-				FSnappingUtils::SnapPointToGrid( SnappedDraggedTo, SnapGridBase );
-			}
+			const bool bShouldConstrainMovement = true;
+			SnappedDraggedTo = SnapLocation(bLocalSpaceSnapping, DesiredGizmoLocation, GizmoStartTransform, SnapGridBase, bShouldConstrainMovement, OriginalConstrainedDragDelta);
 		}
 
 		// Two passes.  First update the real transform.  Then update the unsnapped transform.
@@ -1435,21 +1563,16 @@ void UViewportWorldInteraction::UpdateDragging(
 			const bool bIsUpdatingUnsnappedTarget = ( PassIndex == 1 );
 			const FVector& PassDraggedTo = bIsUpdatingUnsnappedTarget ? ConstrainedDragDeltaFromStart : SnappedDraggedTo;
 
-			if ( InteractionType == ETransformGizmoInteractionType::Translate ||
+			if( InteractionType == ETransformGizmoInteractionType::Translate ||
 				InteractionType == ETransformGizmoInteractionType::TranslateOnPlane )
 			{
-				// Translate the objects!
-				for ( TUniquePtr<FViewportTransformable>& TransformablePtr : Transformables )
-				{
-					FViewportTransformable& Transformable = *TransformablePtr;
-					FTransform& PassTargetTransform = bIsUpdatingUnsnappedTarget ? Transformable.UnsnappedTargetTransform : Transformable.TargetTransform;
+				// Translate the gizmo!
+				FTransform& PassGizmoTargetTransform = bIsUpdatingUnsnappedTarget ? GizmoUnsnappedTargetTransform : GizmoTargetTransform;
 
-					const FVector OldLocation = PassTargetTransform.GetLocation();
-					PassTargetTransform.SetLocation( ( Transformable.StartTransform.GetLocation() - GizmoStartTransform.GetLocation() ) + PassDraggedTo );
+				PassGizmoTargetTransform.SetLocation( PassDraggedTo );
 
-					bMovedAnyTransformables = true;
-					bShouldApplyVelocitiesFromDrag = true;
-				}
+				bMovedTransformGizmo = true;
+				bShouldApplyVelocitiesFromDrag = true;
 			}
 			else if ( InteractionType == ETransformGizmoInteractionType::StretchAndReposition )
 			{
@@ -1485,34 +1608,33 @@ void UViewportWorldInteraction::UpdateDragging(
 				const FVector NewGizmoLocalSize = NewGizmoLocalBounds.GetSize();
 
 				FVector NewGizmoLocalScaleFromStart = FVector( 1.0f );
-				if ( !FMath::IsNearlyZero( NewGizmoLocalSize.GetAbsMin() ) )
-				{
-					NewGizmoLocalScaleFromStart = NewGizmoLocalSize / GizmoStartLocalSize;
-				}
-				else
+				for( int32 AxisIndex = 0; AxisIndex < 3; ++AxisIndex )
 				{
-					// Zero scale.  This is allowed in Unreal, for better or worse.
-					NewGizmoLocalScaleFromStart = FVector( 0 );
+					if( !FMath::IsNearlyZero( GizmoStartLocalSize[AxisIndex] ) )
+					{
+						NewGizmoLocalScaleFromStart[AxisIndex] = NewGizmoLocalSize[AxisIndex] / GizmoStartLocalSize[AxisIndex];
+					}
+					else
+					{
+						// Zero scale.  This is allowed in Unreal, for better or worse.
+						NewGizmoLocalScaleFromStart[AxisIndex] = 0.0f;
+					}
 				}
 
-				// Stretch and reposition the objects!
-				for ( TUniquePtr<FViewportTransformable>& TransformablePtr : Transformables )
+				// Stretch and reposition the gizmo!
 				{
-					FViewportTransformable& Transformable = *TransformablePtr;
-
-					FTransform& PassTargetTransform = bIsUpdatingUnsnappedTarget ? Transformable.UnsnappedTargetTransform : Transformable.TargetTransform;
+					FTransform& PassGizmoTargetTransform = bIsUpdatingUnsnappedTarget ? GizmoUnsnappedTargetTransform : GizmoTargetTransform;
 
 					{
-						const FVector GizmoSpaceTransformableStartLocation = GizmoStartTransform.InverseTransformPositionNoScale( Transformable.StartTransform.GetLocation() );
+						const FVector GizmoSpaceTransformableStartLocation = GizmoStartTransform.InverseTransformPositionNoScale( GizmoStartTransform.GetLocation() );
 						const FVector NewGizmoSpaceLocation = ( GizmoSpaceTransformableStartLocation - GizmoSpacePivotLocation ) * NewGizmoLocalScaleFromStart + GizmoSpacePivotLocation;
-						const FVector OldLocation = PassTargetTransform.GetLocation();
-						PassTargetTransform.SetLocation( GizmoStartTransform.TransformPosition( NewGizmoSpaceLocation ) );
+						PassGizmoTargetTransform.SetLocation( GizmoStartTransform.TransformPosition( NewGizmoSpaceLocation ) );
 					}
 
 					// @todo vreditor: This scale is still in gizmo space, but we're setting it in world space
-					PassTargetTransform.SetScale3D( Transformable.StartTransform.GetScale3D() * NewGizmoLocalScaleFromStart );
+					PassGizmoTargetTransform.SetScale3D( GizmoStartTransform.GetScale3D() * NewGizmoLocalScaleFromStart );
 
-					bMovedAnyTransformables = true;
+					bMovedTransformGizmo = true;
 					bShouldApplyVelocitiesFromDrag = false;
 				}
 			}
@@ -1525,10 +1647,11 @@ void UViewportWorldInteraction::UpdateDragging(
 
 				const FVector PassGizmoSpaceDraggedTo = GizmoStartTransform.InverseTransformPositionNoScale( PassDraggedTo );
 
+				const float ScaleSensitivity = VI::ScaleSensitivity->GetFloat();
 				float AddedScaleOnAxis = 0.0f;
 				if ( InteractionType == ETransformGizmoInteractionType::Scale )
 				{
-					AddedScaleOnAxis = PassGizmoSpaceDraggedTo[ FacingAxisIndex ] * VI::ScaleSensitivity->GetFloat();
+					AddedScaleOnAxis = PassGizmoSpaceDraggedTo[ FacingAxisIndex ] * ScaleSensitivity;
 
 					// Invert if we we are scaling on the negative side of the gizmo
 					if ( DraggingTransformGizmoComponent && DraggingTransformGizmoComponent->GetRelativeTransform().GetLocation()[ FacingAxisIndex ] < 0 )
@@ -1540,15 +1663,13 @@ void UViewportWorldInteraction::UpdateDragging(
 				{
 					// Always use Z for uniform scale
 					const FVector RelativeDraggedTo = PassDraggedTo - GizmoStartTransform.GetLocation();
-					AddedScaleOnAxis = RelativeDraggedTo.Z * VI::ScaleSensitivity->GetFloat();
+					AddedScaleOnAxis = RelativeDraggedTo.Z * ScaleSensitivity;
 				}
 
-				// Scale the objects!
-				for ( TUniquePtr<FViewportTransformable>& TransformablePtr : Transformables )
+				// Scale the gizmo!
 				{
-					FViewportTransformable& Transformable = *TransformablePtr;
-					FTransform& PassTargetTransform = bIsUpdatingUnsnappedTarget ? Transformable.UnsnappedTargetTransform : Transformable.TargetTransform;
-					FVector NewTotalScale = Transformable.StartTransform.GetScale3D();
+					FTransform& PassGizmoTargetTransform = bIsUpdatingUnsnappedTarget ? GizmoUnsnappedTargetTransform : GizmoTargetTransform;
+					FVector NewTotalScale = GizmoStartTransform.GetScale3D();
 
 					if ( InteractionType == ETransformGizmoInteractionType::Scale )
 					{
@@ -1565,8 +1686,8 @@ void UViewportWorldInteraction::UpdateDragging(
 						FSnappingUtils::SnapScale( NewTotalScale, FVector::ZeroVector );
 					}
 
-					PassTargetTransform.SetScale3D( NewTotalScale );
-					bMovedAnyTransformables = true;
+					PassGizmoTargetTransform.SetScale3D( NewTotalScale );
+					bMovedTransformGizmo = true;
 					bShouldApplyVelocitiesFromDrag = true;
 				}
 			}
@@ -1611,13 +1732,11 @@ void UViewportWorldInteraction::UpdateDragging(
 
 				const FTransform WorldToGizmo = GizmoStartTransform.Inverse();
 
-				// Rotate (and reposition) the objects!
-				for ( TUniquePtr<FViewportTransformable>& TransformablePtr : Transformables )
+				// Rotate (and reposition) the gizmo!
 				{
-					FViewportTransformable& Transformable = *TransformablePtr;
-					FTransform& PassTargetTransform = bIsUpdatingUnsnappedTarget ? Transformable.UnsnappedTargetTransform : Transformable.TargetTransform;
-
-					const FTransform GizmoSpaceStartTransform = Transformable.StartTransform * WorldToGizmo;
+					// @todo mesheditor: This stuff doesn't make as much sense now that we're transforming the gizmo itself instead of individual objects.
+					// We can simplify this code along with similar code blocks in this function!
+					const FTransform GizmoSpaceStartTransform = GizmoStartTransform * WorldToGizmo;
 					FTransform GizmoSpaceRotatedTransform = GizmoSpaceStartTransform * RotationDeltaFromStart;
 
 					// Snap rotation in gizmo space
@@ -1630,28 +1749,10 @@ void UViewportWorldInteraction::UpdateDragging(
 
 					const FTransform WorldSpaceRotatedTransform = GizmoSpaceRotatedTransform * GizmoStartTransform;
 
-					PassTargetTransform = WorldSpaceRotatedTransform;
-					bMovedAnyTransformables = true;
+					FTransform& PassGizmoTargetTransform = bIsUpdatingUnsnappedTarget ? GizmoUnsnappedTargetTransform : GizmoTargetTransform;
+					PassGizmoTargetTransform = WorldSpaceRotatedTransform;
+					bMovedTransformGizmo = true;
 					bShouldApplyVelocitiesFromDrag = false;
-
-					// Snap location in world space
-					if ( !bIsUpdatingUnsnappedTarget )
-					{
-						FVector SnappedLocation = PassTargetTransform.GetLocation();
-
-						// Snap in local space, if we need to
-						if ( bLocalSpaceSnapping )
-						{
-							SnappedLocation = GizmoStartTransform.InverseTransformPositionNoScale( SnappedLocation );
-							FSnappingUtils::SnapPointToGrid( SnappedLocation, SnapGridBase );
-							SnappedLocation = GizmoStartTransform.TransformPositionNoScale( SnappedLocation );
-						}
-						else
-						{
-							FSnappingUtils::SnapPointToGrid( SnappedLocation, SnapGridBase );
-						}
-						PassTargetTransform.SetLocation( SnappedLocation );
-					}
 				}
 			}
 			else if ( InteractionType == ETransformGizmoInteractionType::RotateOnAngle )
@@ -1667,41 +1768,43 @@ void UViewportWorldInteraction::UpdateDragging(
 				{
 					const FTransform WorldToGizmo = GizmoStartTransform.Inverse();
 
-					// Rotate the objects!
-					for ( TUniquePtr<FViewportTransformable>& TransformablePtr : Transformables )
+					FTransform NewGizmoToWorld;
 					{
-						FViewportTransformable& Transformable = *TransformablePtr;
-						FTransform& PassTargetTransform = bIsUpdatingUnsnappedTarget ? Transformable.UnsnappedTargetTransform : Transformable.TargetTransform;
-
-						const FVector RotationHandlePlaneLocation =  Transformable.StartTransform.GetLocation();
-						if ( !StartDragHandleDirection.IsSet() )
+						if ( !DraggingRotationHandleDirection.IsSet() )
 						{
-							StartDragHandleDirection = DraggingTransformGizmoComponent->GetComponentTransform().GetRotation().Vector();
-							StartDragHandleDirection->Normalize();
+							DraggingRotationHandleDirection = DraggingTransformGizmoComponent->GetComponentTransform().GetRotation().Vector();
+							DraggingRotationHandleDirection->Normalize();
 						}
 
-						// Get the laserpointer intersection on the plane of the handle
-						FPlane RotationHandlePlane = FPlane( RotationHandlePlaneLocation, StartDragHandleDirection.GetValue() );
+						// Get the laser pointer intersection on the plane of the handle
+						const FPlane RotationPlane = FPlane( GizmoStartTransform.GetLocation(), DraggingRotationHandleDirection.GetValue() );
 						
-						const FVector IntersectWorldLocationOnPlane = FMath::LinePlaneIntersection( LaserPointerStart, LaserPointerStart + LaserPointerDirection, RotationHandlePlane);
+						const FVector LaserImpactOnRotationPlane = FMath::LinePlaneIntersection( LaserPointerStart, LaserPointerStart + LaserPointerDirection, RotationPlane);
 
-						// Set output for hover point
-						OutUnsnappedDraggedTo = IntersectWorldLocationOnPlane;
+						{
+							FTransform GizmoTransformNoRotation = FTransform(FRotator::ZeroRotator, GizmoStartTransform.GetLocation());
+							const ECoordSystem CoordSystem = GLevelEditorModeTools().GetCoordSystem(true);
+							if (CoordSystem == COORD_Local)
+							{
+								GizmoTransformNoRotation.SetRotation(GizmoStartTransform.GetRotation());
+							}
+
+							LocalIntersectPointOnRotationGizmo = GizmoTransformNoRotation.InverseTransformPositionNoScale(LaserImpactOnRotationPlane);
+						}
 
-						// The starting transformation of this transformable in gizmo space
-						const FTransform GizmoSpaceStartTransform = Transformable.StartTransform * WorldToGizmo;
+						// Set output for hover point
+						OutUnsnappedDraggedTo = LaserImpactOnRotationPlane;
 
 						// Relative offset of the intersection on the plane
-						const FVector IntersectLocalLocationOnPlane = Transformable.StartTransform.InverseTransformPositionNoScale( IntersectWorldLocationOnPlane );
+						const FVector GizmoSpaceLaserImpactOnRotationPlane = WorldToGizmo.TransformPosition( LaserImpactOnRotationPlane );
 						FVector RotatedIntersectLocationOnPlane;
 						if ( GLevelEditorModeTools().GetCoordSystem( true ) == COORD_Local )
 						{
-							RotatedIntersectLocationOnPlane = GizmoSpaceFacingAxisVector.Rotation().UnrotateVector( IntersectLocalLocationOnPlane );
+							RotatedIntersectLocationOnPlane = GizmoSpaceFacingAxisVector.Rotation().UnrotateVector( GizmoSpaceLaserImpactOnRotationPlane );
 						}
 						else
 						{
-							const FVector TransformableSpaceFacingAxisVector = Transformable.StartTransform.InverseTransformVectorNoScale( GizmoStartTransform.InverseTransformVectorNoScale( GizmoSpaceFacingAxisVector ) );
-							RotatedIntersectLocationOnPlane = TransformableSpaceFacingAxisVector.Rotation().UnrotateVector( IntersectLocalLocationOnPlane );
+							RotatedIntersectLocationOnPlane = GizmoStartTransform.TransformVector( GizmoSpaceFacingAxisVector ).Rotation().UnrotateVector( GizmoSpaceLaserImpactOnRotationPlane );
 						}
 
 						// Get the angle between the center and the intersected point
@@ -1715,22 +1818,22 @@ void UViewportWorldInteraction::UpdateDragging(
 						const float AngleDeltaRotationFromStart = FMath::FindDeltaAngleRadians( AngleToIntersectedLocation, StartDragAngleOnRotation.GetValue() );
 						const FQuat GizmoSpaceDeltaRotation = FQuat( GizmoSpaceFacingAxisVector, AngleDeltaRotationFromStart );
 
-						// Add the delta rotation to the starting transformation
-						FTransform GizmoSpaceRotatedTransform = GizmoSpaceStartTransform * GizmoSpaceDeltaRotation;
-
 						// Snap rotation in gizmo space
+						FTransform GizmoSpaceRotatedTransform( GizmoSpaceDeltaRotation );
 						if ( !bIsUpdatingUnsnappedTarget )
 						{
 							FRotator SnappedRotation = GizmoSpaceRotatedTransform.GetRotation().Rotator();
 							FSnappingUtils::SnapRotatorToGrid( SnappedRotation );
 							GizmoSpaceRotatedTransform.SetRotation( SnappedRotation.Quaternion() );
 						}
+						NewGizmoToWorld = GizmoSpaceRotatedTransform * GizmoStartTransform;
+					}
 
-						// Rotation back to world space 
-						const FTransform WorldSpaceRotatedTransform = GizmoSpaceRotatedTransform * GizmoStartTransform;
-
-						PassTargetTransform = WorldSpaceRotatedTransform;
-						bMovedAnyTransformables = true;
+					// Rotate the gizmo!
+					{
+						FTransform& PassTargetTransform = bIsUpdatingUnsnappedTarget ? GizmoUnsnappedTargetTransform : GizmoTargetTransform;
+						PassTargetTransform = NewGizmoToWorld;
+						bMovedTransformGizmo = true;
 						bShouldApplyVelocitiesFromDrag = true;
 					}
 				}
@@ -1770,7 +1873,7 @@ void UViewportWorldInteraction::UpdateDragging(
 			// side of the line that moved the most this update.
 			const float TotalDistance = LineStartDistance + LineEndDistance;
 			float LineStartToEndActivityWeight = 0.5f;	// Default to right in the center, if no distance moved yet.
-			if ( !FMath::IsNearlyZero( TotalDistance ) )	// Avoid division by zero
+			if ( VI::ScaleWorldWithDynamicPivot->GetInt() != 0 && !FMath::IsNearlyZero( TotalDistance ) )	// Avoid division by zero
 			{
 				LineStartToEndActivityWeight = LineStartDistance / TotalDistance;
 			}
@@ -1821,92 +1924,41 @@ void UViewportWorldInteraction::UpdateDragging(
 
 		if ( DraggingMode == EViewportInteractionDraggingMode::TransformablesFreely )
 		{
-			for ( TUniquePtr<FViewportTransformable>& TransformablePtr : Transformables )
+			if( !bHasPivotLocation )
 			{
-				FViewportTransformable& Transformable = *TransformablePtr;
-
-				const FTransform OldTransform = Transformable.UnsnappedTargetTransform;
-
-				// Two passes.  First update the real transform.  Then update the unsnapped transform.
-				for ( int32 PassIndex = 0; PassIndex < 2; ++PassIndex )
-				{
-					const bool bIsUpdatingUnsnappedTarget = ( PassIndex == 1 );
-
-					// Scale snap!
-					FVector SnappedScaleOffset = ScaleOffset;
-					if ( !bIsUpdatingUnsnappedTarget && FSnappingUtils::IsScaleSnapEnabled() )
-					{
-						// Did scale even change?
-						if ( !( FMath::IsNearlyEqual( ScaleOffset.GetAbsMax(), 1.0f ) && FMath::IsNearlyEqual( ScaleOffset.GetAbsMin(), 1.0f ) ) )
-						{
-							FVector NewTotalScale = OldTransform.GetScale3D() * ScaleOffset;
-							FSnappingUtils::SnapScale( NewTotalScale, FVector::ZeroVector );
-							SnappedScaleOffset = NewTotalScale / OldTransform.GetScale3D();
-						}
-					}
-
-					const FTransform ScaleOffsetTransform( FQuat::Identity, FVector::ZeroVector, SnappedScaleOffset );
-					const FTransform TranslationOffsetTransform( FQuat::Identity, TranslationOffset );
-					const FTransform RotationOffsetTransform( RotationOffset, FVector::ZeroVector );
+				PivotLocation = GizmoUnsnappedTargetTransform.GetLocation();
+				bHasPivotLocation = true;
+			}
 
-					if ( !bHasPivotLocation )
-					{
-						PivotLocation = OldTransform.GetLocation();
-					}
+			FTransform NewGizmoToWorld = FTransform::Identity;
+			{
+				const FTransform ScaleOffsetTransform( FQuat::Identity, FVector::ZeroVector, ScaleOffset );
+				const FTransform TranslationOffsetTransform( FQuat::Identity, TranslationOffset );
+				const FTransform RotationOffsetTransform( RotationOffset, FVector::ZeroVector );
 
-					// @todo vreditor multi: This is only solving for rotating/scaling one object.  Really we want to rotate/scale the GIZMO and have all objects rotated and repositioned within it!
-					FTransform NewTransform;
-					const bool bIgnoreGizmoSpace = true;
-					if ( bIgnoreGizmoSpace )
-					{
-						const FTransform PivotToWorld = FTransform( FQuat::Identity, PivotLocation );
-						const FTransform WorldToPivot = FTransform( FQuat::Identity, -PivotLocation );
-						NewTransform = OldTransform * WorldToPivot * ScaleOffsetTransform * RotationOffsetTransform * PivotToWorld * TranslationOffsetTransform;
-					}
-					else
-					{
-						// @todo vreditor perf: Can move most of this outside of the inner loop.  Also could be optimized by reducing the number of transforms.
-						const FTransform GizmoToWorld = GizmoStartTransform;
-						const FTransform WorldToGizmo = GizmoToWorld.Inverse();
-						const FVector GizmoSpacePivotLocation = WorldToGizmo.TransformPositionNoScale( PivotLocation );
-						const FTransform GizmoToPivot = FTransform( FQuat::Identity, -GizmoSpacePivotLocation );
-						const FTransform PivotToGizmo = FTransform( FQuat::Identity, GizmoSpacePivotLocation );
-						const FQuat GizmoSpaceRotationOffsetQuat = WorldToGizmo.GetRotation() * RotationOffset;	// NOTE: FQuat composition order is opposite that of FTransform!
-						const FTransform GizmoSpaceRotationOffsetTransform( GizmoSpaceRotationOffsetQuat, FVector::ZeroVector );
-						NewTransform = OldTransform * TranslationOffsetTransform * WorldToGizmo * GizmoToPivot * ScaleOffsetTransform * GizmoSpaceRotationOffsetTransform * PivotToGizmo * GizmoToWorld;
-					}
+				const FTransform PivotToWorld = FTransform( FQuat::Identity, PivotLocation );
+				const FTransform WorldToPivot = FTransform( FQuat::Identity, -PivotLocation );
+				NewGizmoToWorld = GizmoUnsnappedTargetTransform * WorldToPivot * ScaleOffsetTransform * RotationOffsetTransform * PivotToWorld * TranslationOffsetTransform;
+			}
 
-					if ( bIsUpdatingUnsnappedTarget )
-					{
-						const FVector OldLocation = Transformable.UnsnappedTargetTransform.GetLocation();
-						Transformable.UnsnappedTargetTransform = NewTransform;
-					}
-					else
-					{
-						Transformable.TargetTransform = NewTransform;
+			// Grid snap!
+			FTransform SnappedNewGizmoToWorld = NewGizmoToWorld;
+			{
+				const FVector GizmoDragDelta = NewGizmoToWorld.GetLocation() - GizmoUnsnappedTargetTransform.GetLocation();
+				const bool bShouldConstrainMovement = false;
+				SnappedNewGizmoToWorld.SetLocation( SnapLocation( bLocalSpaceSnapping, NewGizmoToWorld.GetLocation(), GizmoUnsnappedTargetTransform, SnapGridBase, bShouldConstrainMovement, GizmoDragDelta) );
+			}
 
-						// Grid snap!
-						FVector SnappedLocation = Transformable.TargetTransform.GetLocation();
-						if ( bLocalSpaceSnapping )
-						{
-							SnappedLocation = GizmoStartTransform.InverseTransformPositionNoScale( SnappedLocation );
-							FSnappingUtils::SnapPointToGrid( SnappedLocation, SnapGridBase );
-							SnappedLocation = GizmoStartTransform.TransformPositionNoScale( SnappedLocation );
-						}
-						else
-						{
-							FSnappingUtils::SnapPointToGrid( SnappedLocation, SnapGridBase );
-						}
-						Transformable.TargetTransform.SetLocation( SnappedLocation );
+			// Two passes.  First update the real transform.  Then update the unsnapped transform.
+			for( int32 PassIndex = 0; PassIndex < 2; ++PassIndex )
+			{
+				const bool bIsUpdatingUnsnappedTarget = ( PassIndex == 1 );
 
-						// Rotation snap!
-						FRotator SnappedRotation = Transformable.TargetTransform.GetRotation().Rotator();
-						FSnappingUtils::SnapRotatorToGrid( SnappedRotation );
-						Transformable.TargetTransform.SetRotation( SnappedRotation.Quaternion() );
-					}
-				}
+				const FTransform NewTransform = ( bIsUpdatingUnsnappedTarget ? NewGizmoToWorld : SnappedNewGizmoToWorld );
 
-				bMovedAnyTransformables = true;
+				FTransform& PassGizmoTargetTransform = bIsUpdatingUnsnappedTarget ? GizmoUnsnappedTargetTransform : GizmoTargetTransform;
+				PassGizmoTargetTransform = NewTransform;
+				bMovedTransformGizmo = true;
 				bShouldApplyVelocitiesFromDrag = true;
 			}
 		}
@@ -1922,25 +1974,11 @@ void UViewportWorldInteraction::UpdateDragging(
 				const float NewWorldToMetersScale = OldWorldToMetersScale / WorldScaleOffset;
 
 				// NOTE: Instead of clamping, we simply skip changing the W2M this frame if it's out of bounds.  Clamping makes our math more complicated.
-				if ( NewWorldToMetersScale != OldWorldToMetersScale &&
+				if ( !FMath::IsNearlyEqual(NewWorldToMetersScale, OldWorldToMetersScale) &&
 					NewWorldToMetersScale >= GetMinScale() && NewWorldToMetersScale <= GetMaxScale() )
 				{
-					SetWorldToMetersScale( NewWorldToMetersScale );
-
-					// Because the tracking space size has changed, but our head position within that space relative to the origin
-					// of the room is the same (before scaling), we need to offset our location within the tracking space to compensate.
-					// This makes the user feel like their head and hands remain in the same location.
-					const FVector RoomSpacePivotLocation = PivotLocation;
-					const FVector WorldSpacePivotLocation = RoomTransform.TransformPosition( RoomSpacePivotLocation );
-
-					const FVector NewRoomSpacePivotLocation = ( RoomSpacePivotLocation / OldWorldToMetersScale ) * NewWorldToMetersScale;
-					const FVector NewWorldSpacePivotLocation = RoomTransform.TransformPosition( NewRoomSpacePivotLocation );
-
-					const FVector WorldSpacePivotDelta = ( NewWorldSpacePivotLocation - WorldSpacePivotLocation );
-
-					const FVector NewWorldSpaceRoomLocation = RoomTransform.GetLocation() - WorldSpacePivotDelta;
-
-					RoomTransform.SetLocation( NewWorldSpaceRoomLocation );
+					SetWorldToMetersScale(NewWorldToMetersScale);
+					CompensateRoomTransformForWorldScale(RoomTransform, NewWorldToMetersScale, PivotLocation);
 				}
 			}
 
@@ -1964,53 +2002,51 @@ void UViewportWorldInteraction::UpdateDragging(
 				RoomTransform = TranslationOffsetTransform.Inverse() * RoomTransform * WorldToPivot * RotationOffsetTransform.Inverse() * PivotToWorld;
 			}
 
-			SetRoomTransform( RoomTransform );
+			RoomTransformToSetOnFrame = MakeTuple( RoomTransform, CurrentTickNumber + 1 );
 		}
-
 	}
 
 
 	// If we're not using smooth snapping, go ahead and immediately update the transforms of all objects.  If smooth
 	// snapping is enabled, this will be done in Tick() instead.
-	if ( bMovedAnyTransformables )
+	if (bMovedTransformGizmo)
 	{
 		// Update velocity if we're simulating in editor
-		if ( GEditor->bIsSimulatingInEditor )
+		if (GEditor->bIsSimulatingInEditor)
 		{
-			for ( TUniquePtr<FViewportTransformable>& TransformablePtr : Transformables )
+			FVector MoveDelta = GizmoUnsnappedTargetTransform.GetLocation() - GizmoLastTransform.GetLocation();
+			if( bShouldApplyVelocitiesFromDrag )
+			{
+				bIsDrivingVelocityOfSimulatedTransformables = true;
+			}
+			else
+			{
+				MoveDelta = FVector::ZeroVector;
+			}
+
+			for (TUniquePtr<FViewportTransformable>& TransformablePtr : Transformables)
 			{
 				FViewportTransformable& Transformable = *TransformablePtr;
 
 				// @todo VREditor physics: When freely transforming, should set angular velocity too (will pivot be the same though??)
-				if( Transformable.IsPhysicallySimulated() )
+				if (Transformable.IsPhysicallySimulated())
 				{
-					FVector MoveDelta = Transformable.UnsnappedTargetTransform.GetLocation() - Transformable.LastTransform.GetLocation();
-					if ( bShouldApplyVelocitiesFromDrag )
-					{
-						bIsDrivingVelocityOfSimulatedTransformables = true;
-					}
-					else
-					{
-						MoveDelta = FVector::ZeroVector;
-					}
-
-					Transformable.SetLinearVelocity( MoveDelta * DeltaTime * VI::InertiaVelocityBoost->GetFloat() );
+					Transformable.SetLinearVelocity( ( MoveDelta * VI::InertiaVelocityBoost->GetFloat() ) / DeltaTime );
 				}
 			}
 		}
 
 		const bool bSmoothSnappingEnabled = IsSmoothSnappingEnabled();
-		if ( !bSmoothSnappingEnabled && !bIsInterpolatingTransformablesFromSnapshotTransform )
+		if (!bSmoothSnappingEnabled && !bIsInterpolatingTransformablesFromSnapshotTransform)
 		{
+			GizmoLastTransform = GizmoTargetTransform;
 			const bool bSweep = bShouldApplyVelocitiesFromDrag && VI::SweepPhysicsWhileSimulating->GetInt() != 0 && bIsDrivingVelocityOfSimulatedTransformables;
-			for ( TUniquePtr<FViewportTransformable>& TransformablePtr : Transformables )
+			for (TUniquePtr<FViewportTransformable>& TransformablePtr : Transformables)
 			{
 				FViewportTransformable& Transformable = *TransformablePtr;
 
-				Transformable.LastTransform = Transformable.TargetTransform;
-
 				// Update the transform!
-				Transformable.ApplyTransform( Transformable.TargetTransform, bSweep );
+				Transformable.ApplyTransform(Transformable.StartTransform * GizmoStartTransform.Inverse() * GizmoTargetTransform , bSweep);
 			}
 		}
 	}
@@ -2018,6 +2054,7 @@ void UViewportWorldInteraction::UpdateDragging(
 	bIsFirstDragUpdate = false;
 	bSkipInteractiveWorldMovementThisFrame = false;
 }
+
 END_FUNCTION_BUILD_OPTIMIZATION
 
 FVector UViewportWorldInteraction::ComputeConstrainedDragDeltaFromStart( 
@@ -2168,115 +2205,142 @@ FVector UViewportWorldInteraction::ComputeConstrainedDragDeltaFromStart(
 	return ConstrainedWorldSpaceDeltaFromStart;
 }
 
-bool UViewportWorldInteraction::StartDragging( UViewportInteractor* Interactor, UActorComponent* ClickedTransformGizmoComponent, const FVector& HitLocation, const bool bIsPlacingNewObjects, const bool bStartTransaction )
+void UViewportWorldInteraction::StartDragging( UViewportInteractor* Interactor, UActorComponent* ClickedTransformGizmoComponent, const FVector& HitLocation, const bool bIsPlacingNewObjects, const bool bAllowInterpolationWhenPlacing, const bool bStartTransaction, const bool bWithGrabberSphere )
 {
-	bool bStartedDragging = false;
-	{
-		FVector LaserPointerStart, LaserPointerEnd;
-		const bool bHaveLaserPointer = Interactor->GetLaserPointer( /* Out */ LaserPointerStart, /* Out */ LaserPointerEnd );
-		if ( bHaveLaserPointer )
-		{
-			FViewportInteractorData& InteractorData = Interactor->GetInteractorData();
+	bool bHaveGrabberSphere = false;
+	bool bHaveLaserPointer = false;
+	FSphere GrabberSphere = FSphere( 0 );
+	FVector LaserPointerStart = FVector::ZeroVector;
+	FVector LaserPointerEnd = FVector::ZeroVector;
 
-			if( bStartTransaction )
-			{
-				// Capture undo state
-				TrackingTransaction.TransCount++;
-				TrackingTransaction.Begin( LOCTEXT( "MovingObjects", "Move Object" ) );
+	if( bWithGrabberSphere )
+	{
+		bHaveGrabberSphere = Interactor->GetGrabberSphere( /* Out */ GrabberSphere );
+		check( bHaveGrabberSphere );
+	}
+	else
+	{
+		bHaveLaserPointer = Interactor->GetLaserPointer( /* Out */ LaserPointerStart, /* Out */ LaserPointerEnd );
+		check( bHaveLaserPointer );
+	}
 
-				// Suspend actor/component modification during each delta step to avoid recording unnecessary overhead into the transaction buffer
-				GEditor->DisableDeltaModification( true );
-			}
-			else
-			{
-				// A transaction should already be in progress
-				check( GEditor->IsTransactionActive() );
-			}
+	FViewportInteractorData& InteractorData = Interactor->GetInteractorData();
 
-			// Give the interactor a chance to do something when starting dragging
-			Interactor->OnStartDragging( HitLocation, bIsPlacingNewObjects );
+	if( bStartTransaction )
+	{
+		// Capture undo state
+		TrackingTransaction.TransCount++;
+		TrackingTransaction.Begin( LOCTEXT( "MovingObjects", "Move Object" ) );
+	}
+	else
+	{
+		// Did you forget to use an FScopedTransaction?  If GUndo was null, then most likely we forgot to wrap this call within an editor transaction.
+		// The only exception is in Simulate mode, where Undo is not allowed.
+		check( GUndo != nullptr || GEditor == nullptr || GEditor->bIsSimulatingInEditor );
+	}
+	
+	// Suspend actor/component modification during each delta step to avoid recording unnecessary overhead into the transaction buffer
+	GEditor->DisableDeltaModification( true );
 
-			const bool bUsingGizmo = ClickedTransformGizmoComponent != nullptr;
-			check( !bUsingGizmo || ( ClickedTransformGizmoComponent->GetOwner() == TransformGizmoActor ) );
+	// Give the interactor a chance to do something when starting dragging
+	Interactor->OnStartDragging( HitLocation, bIsPlacingNewObjects );
 
-			// Start dragging the objects right away!
-			InteractorData.DraggingMode = InteractorData.LastDraggingMode = bIsPlacingNewObjects ? EViewportInteractionDraggingMode::TransformablesAtLaserImpact :
-				( bUsingGizmo ? EViewportInteractionDraggingMode::TransformablesWithGizmo : EViewportInteractionDraggingMode::TransformablesFreely );
+	const bool bUsingGizmo = ClickedTransformGizmoComponent != nullptr;
+	check( !bUsingGizmo || ( ClickedTransformGizmoComponent->GetOwner() == TransformGizmoActor ) );
 
-			bAreTransformablesMoving = true;
+	const bool bShouldUseLaserImpactDrag = bIsPlacingNewObjects;
+	// Start dragging the objects right away!
+	InteractorData.DraggingMode = InteractorData.LastDraggingMode = bShouldUseLaserImpactDrag ? EViewportInteractionDraggingMode::TransformablesAtLaserImpact :
+		( bUsingGizmo ? EViewportInteractionDraggingMode::TransformablesWithGizmo : EViewportInteractionDraggingMode::TransformablesFreely );
 
-			// Starting a new drag, so make sure the other hand doesn't think it's assisting us
-			if( Interactor->GetOtherInteractor() != nullptr )
-			{
-				FViewportInteractorData& OtherInteractorData = Interactor->GetOtherInteractor()->GetInteractorData();
-				OtherInteractorData.bWasAssistingDrag = false;
-			}
+	bAreTransformablesMoving = true;
 
-			InteractorData.bIsFirstDragUpdate = true;
-			InteractorData.bWasAssistingDrag = false;
-			InteractorData.DragRayLength = ( HitLocation - LaserPointerStart ).Size();
-			InteractorData.LastLaserPointerStart = LaserPointerStart;
-			InteractorData.LastDragToLocation = HitLocation;
-			InteractorData.LaserPointerImpactAtDragStart = HitLocation;
-			InteractorData.DragTranslationVelocity = FVector::ZeroVector;
-			InteractorData.DragRayLengthVelocity = 0.0f;
-			InteractorData.bIsDrivingVelocityOfSimulatedTransformables = false;
+	// Starting a new drag, so make sure the other hand doesn't think it's assisting us
+	if( Interactor->GetOtherInteractor() != nullptr )
+	{
+		FViewportInteractorData& OtherInteractorData = Interactor->GetOtherInteractor()->GetInteractorData();
+		OtherInteractorData.bWasAssistingDrag = false;
+	}
 
-			// Start dragging the transform gizmo.  Even if the user clicked on the actor itself, we'll use
-			// the gizmo to transform it.
-			if ( bUsingGizmo )
-			{
-				InteractorData.DraggingTransformGizmoComponent = Cast<USceneComponent>( ClickedTransformGizmoComponent );
-			}
-			else
-			{
-				InteractorData.DraggingTransformGizmoComponent = nullptr;
-			}
+	InteractorData.bDraggingWithGrabberSphere = bWithGrabberSphere;
+	InteractorData.bIsFirstDragUpdate = true;
+	InteractorData.bWasAssistingDrag = false;
+	if( bWithGrabberSphere )
+	{
+		// NOTE: With the grabber sphere, we don't allow the user to shift the drag ray distance
+		InteractorData.DragRayLength = 0.0f;
+	}
+	else
+	{
+		InteractorData.DragRayLength = ( HitLocation - LaserPointerStart ).Size();
+	}
+	InteractorData.LastDragToLocation = HitLocation;
+	InteractorData.InteractorRotationAtDragStart = InteractorData.Transform.GetRotation();
+	InteractorData.GrabberSphereLocationAtDragStart = bWithGrabberSphere ? GrabberSphere.Center : FVector::ZeroVector;
+	InteractorData.ImpactLocationAtDragStart = HitLocation;
+	InteractorData.DragTranslationVelocity = FVector::ZeroVector;
+	InteractorData.DragRayLengthVelocity = 0.0f;
+	InteractorData.bIsDrivingVelocityOfSimulatedTransformables = false;
 
-			if ( TransformGizmoActor != nullptr )
-			{
-				InteractorData.TransformGizmoInteractionType = TransformGizmoActor->GetInteractionType( InteractorData.DraggingTransformGizmoComponent.Get(), InteractorData.OptionalHandlePlacement );
-				InteractorData.GizmoStartTransform = TransformGizmoActor->GetTransform();
-				InteractorData.GizmoStartLocalBounds = GizmoLocalBounds;
-			}
-			else
-			{
-				InteractorData.TransformGizmoInteractionType = ETransformGizmoInteractionType::None;
-				InteractorData.GizmoStartTransform = FTransform::Identity;
-				InteractorData.GizmoStartLocalBounds = FBox( 0 );
-			}
-			InteractorData.GizmoSpaceFirstDragUpdateOffsetAlongAxis = FVector::ZeroVector;	// Will be determined on first update
-			InteractorData.GizmoSpaceDragDeltaFromStartOffset = FVector::ZeroVector;	// Set every frame while dragging
+	// Start dragging the transform gizmo.  Even if the user clicked on the actor itself, we'll use
+	// the gizmo to transform it.
+	if ( bUsingGizmo )
+	{
+		InteractorData.DraggingTransformGizmoComponent = Cast<USceneComponent>( ClickedTransformGizmoComponent );
+	}
+	else
+	{
+		InteractorData.DraggingTransformGizmoComponent = nullptr;
+	}
 
-			bDraggedSinceLastSelection = true;
-			LastDragGizmoStartTransform = InteractorData.GizmoStartTransform;
+	if ( TransformGizmoActor != nullptr )
+	{
+		InteractorData.TransformGizmoInteractionType = TransformGizmoActor->GetInteractionType( InteractorData.DraggingTransformGizmoComponent.Get(), InteractorData.OptionalHandlePlacement );
+		InteractorData.GizmoStartTransform = TransformGizmoActor->GetTransform();
+		InteractorData.GizmoStartLocalBounds = GizmoLocalBounds;
+	}
+	else
+	{
+		InteractorData.TransformGizmoInteractionType = ETransformGizmoInteractionType::None;
+		InteractorData.GizmoStartTransform = FTransform::Identity;
+		InteractorData.GizmoStartLocalBounds = FBox( 0 );
+	}
+	InteractorData.GizmoLastTransform = InteractorData.GizmoTargetTransform = InteractorData.GizmoUnsnappedTargetTransform = InteractorData.GizmoInterpolationSnapshotTransform = InteractorData.GizmoStartTransform;
+	InteractorData.GizmoSpaceFirstDragUpdateOffsetAlongAxis = FVector::ZeroVector;	// Will be determined on first update
+	InteractorData.GizmoSpaceDragDeltaFromStartOffset = FVector::ZeroVector;	// Set every frame while dragging
 
-			// Make sure all of our transformables are reset to their initial transform before we start dragging them
-			ResetTransformables();
+	bDraggedSinceLastSelection = true;
+	LastDragGizmoStartTransform = InteractorData.GizmoStartTransform;
 
-			// Calculate the offset between the average location and the hit location after setting the transformables for the new selected objects
-			InteractorData.StartHitLocationToTransformableCenter = CalculateAverageLocationOfTransformables() - HitLocation;
+	// Make sure all of our transformables are reset to their initial transform before we start dragging them
+	for( TUniquePtr<FViewportTransformable>& Transformable : Transformables )
+	{
+		Transformable->StartTransform = Transformable->GetTransform();
+	}
 
-			// If we're placing actors, start interpolating to their actual location.  This helps smooth everything out when
-			// using the laser impact point as the target transform
-			if( bIsPlacingNewObjects )
-			{
-				bIsInterpolatingTransformablesFromSnapshotTransform = true;
-				const FTimespan CurrentTime = FTimespan::FromSeconds( FPlatformTime::Seconds() );
-				TransformablesInterpolationStartTime = CurrentTime;
-				TransformablesInterpolationDuration = VI::PlacementInterpolationDuration->GetFloat();
-			}
+	// Calculate the offset between the average location and the hit location after setting the transformables for the new selected objects
+	InteractorData.StartHitLocationToTransformableCenter = CalculateAverageLocationOfTransformables() - HitLocation;
 
-			// Play a haptic effect when objects are picked up
-			Interactor->PlayHapticEffect( Interactor->GetDragHapticFeedbackStrength() );
+	// If we're placing actors, start interpolating to their actual location.  This helps smooth everything out when
+	// using the laser impact point as the target transform
+	if( bIsPlacingNewObjects && bAllowInterpolationWhenPlacing )
+	{
+		bIsInterpolatingTransformablesFromSnapshotTransform = true;
+		bFreezePlacementWhileInterpolatingTransformables = false;	// Always update the placement location when dragging out new objects
+		const FTimespan CurrentTime = FTimespan::FromSeconds( FPlatformTime::Seconds() );
+		TransformablesInterpolationStartTime = CurrentTime;
+		TransformablesInterpolationDuration = VI::PlacementInterpolationDuration->GetFloat();
+	}
 
-			OnStartDraggingEvent.Broadcast( Interactor );
+	// Play a haptic effect when objects are picked up
+	Interactor->PlayHapticEffect( Interactor->GetDragHapticFeedbackStrength() );
 
-			bStartedDragging = true;
-		}
+	if (ViewportTransformer)
+	{
+		ViewportTransformer->OnStartDragging(Interactor);
 	}
 
-	return bStartedDragging;
+	OnStartDraggingEvent.Broadcast( Interactor );
 }
 
 
@@ -2302,6 +2366,13 @@ void UViewportWorldInteraction::StopDragging( UViewportInteractor* Interactor )
 			OtherInteractorData->DraggingMode = OtherInteractorData->LastDraggingMode = InteractorData.DraggingMode;
 			OtherInteractorData->bIsDrivingVelocityOfSimulatedTransformables = InteractorData.bIsDrivingVelocityOfSimulatedTransformables;
 
+			OtherInteractorData->GizmoStartTransform = InteractorData.GizmoStartTransform;
+			OtherInteractorData->GizmoLastTransform = InteractorData.GizmoLastTransform;
+			OtherInteractorData->GizmoTargetTransform = InteractorData.GizmoTargetTransform;
+			OtherInteractorData->GizmoUnsnappedTargetTransform = InteractorData.GizmoUnsnappedTargetTransform;
+			OtherInteractorData->GizmoInterpolationSnapshotTransform = InteractorData.GizmoInterpolationSnapshotTransform;
+			OtherInteractorData->GizmoStartLocalBounds = InteractorData.GizmoStartLocalBounds;
+
 			// The other hand is no longer assisting, as it's now the primary interacting hand.
 			OtherInteractorData->bWasAssistingDrag = false;
 
@@ -2332,7 +2403,7 @@ void UViewportWorldInteraction::StopDragging( UViewportInteractor* Interactor )
 					 InteractorData.TransformGizmoInteractionType == ETransformGizmoInteractionType::RotateOnAngle )
 				{
 					StartDragAngleOnRotation.Reset();
-					StartDragHandleDirection.Reset();
+					DraggingRotationHandleDirection.Reset();
 				}
 
 				// If we're not dragging anything around, check to see if transformables have come to a rest.
@@ -2346,13 +2417,16 @@ void UViewportWorldInteraction::StopDragging( UViewportInteractor* Interactor )
 					// Finalize the drag
 					FinishedMovingTransformables();
 				}	
+
+				if (ViewportTransformer)
+				{
+					ViewportTransformer->OnStopDragging(Interactor);
+				}
 			}
 		}
-		
+
 		InteractorData.DraggingMode = EViewportInteractionDraggingMode::Nothing;
 	}
-
-	// NOTE: Even though we're done dragging, we keep our list of transformables so that inertial drag can still push them around!
 }
 
 
@@ -2360,6 +2434,7 @@ void UViewportWorldInteraction::FinishedMovingTransformables()
 {
 	bAreTransformablesMoving = false;
 	bIsInterpolatingTransformablesFromSnapshotTransform = false;
+	bFreezePlacementWhileInterpolatingTransformables = false;
 	TransformablesInterpolationStartTime = FTimespan::Zero();
 	TransformablesInterpolationDuration = 1.0f;
 
@@ -2376,9 +2451,14 @@ void UViewportWorldInteraction::FinishedMovingTransformables()
 }
 
 
-void UViewportWorldInteraction::SetUseInputProcessor( bool bInUseInputProcessor )
+void UViewportWorldInteraction::SetUseInputPreprocessor( bool bInUseInputPreprocessor )
+{
+	bUseInputPreprocessor = bInUseInputPreprocessor;
+}
+
+void UViewportWorldInteraction::AllowWorldMovement(bool bAllow)
 {
-	bUseInputProcessor = bInUseInputProcessor;
+	bAllowWorldMovement = bAllow;
 }
 
 bool UViewportWorldInteraction::FindPlacementPointUnderLaser( UViewportInteractor* Interactor, FVector& OutHitLocation )
@@ -2387,6 +2467,7 @@ bool UViewportWorldInteraction::FindPlacementPointUnderLaser( UViewportInteracto
 	bool bHitSomething = false;
 	FVector HitLocation = FVector::ZeroVector;
 
+	bool bAnyPhysicallySimulatedTransformables = false;
 	static TArray<AActor*> IgnoredActors;
 	{
 		IgnoredActors.Reset();
@@ -2397,6 +2478,11 @@ bool UViewportWorldInteraction::FindPlacementPointUnderLaser( UViewportInteracto
 			FViewportTransformable& Transformable = *TransformablePtr;
 
 			Transformable.UpdateIgnoredActorList( IgnoredActors );
+
+			if( Transformable.IsPhysicallySimulated() )
+			{
+				bAnyPhysicallySimulatedTransformables = true;
+			}
 		} 
 	}
 
@@ -2427,7 +2513,11 @@ bool UViewportWorldInteraction::FindPlacementPointUnderLaser( UViewportInteracto
 					ExtremeDirection.Y >= 0.0f ? Box.Max.Y : Box.Min.Y,
 					ExtremeDirection.Z >= 0.0f ? Box.Max.Z : Box.Min.Z );
 				const float ProjectionDistance = FVector::DotProduct( ExtremePoint, ExtremeDirection );
-				ExtremePointOnBox = ExtremeDirection * ProjectionDistance;
+
+				const float ExtraOffset = 
+					( bAnyPhysicallySimulatedTransformables && GEditor->bIsSimulatingInEditor ) ? ( GizmoLocalBounds.GetSize().GetAbsMax() * VI::PlacementOffsetScaleWhileSimulating->GetFloat() ) : 0.0f;
+
+				ExtremePointOnBox = ExtremeDirection * ( ProjectionDistance + ExtraOffset );
 			}
 
 			const FVector WorldSpaceExtremePointOnBox = InteractorData.GizmoStartTransform.TransformVectorNoScale( ExtremePointOnBox );
@@ -2449,10 +2539,21 @@ FTrackingTransaction& UViewportWorldInteraction::GetTrackingTransaction()
 
 bool UViewportWorldInteraction::IsSmoothSnappingEnabled() const
 {
+	// @todo vreditor perf: We could cache this, perhaps.  There are a few other places we do similar checks. O(N).
+	bool bAnyPhysicallySimulatedTransformables = false;
+	for( const TUniquePtr<FViewportTransformable>& TransformablePtr : Transformables )
+	{
+		FViewportTransformable& Transformable = *TransformablePtr;
+		if( Transformable.IsPhysicallySimulated() )
+		{
+			bAnyPhysicallySimulatedTransformables = true;
+		}
+	}
+
 	const float SmoothSnapSpeed = VI::SmoothSnapSpeed->GetFloat();
 	const bool bSmoothSnappingEnabled =
-		!GEditor->bIsSimulatingInEditor &&
-		( FSnappingUtils::IsSnapToGridEnabled() || FSnappingUtils::IsRotationSnapEnabled() || FSnappingUtils::IsScaleSnapEnabled() ) &&
+		( !GEditor->bIsSimulatingInEditor || !bAnyPhysicallySimulatedTransformables ) &&
+		( FSnappingUtils::IsSnapToGridEnabled() || FSnappingUtils::IsRotationSnapEnabled() || FSnappingUtils::IsScaleSnapEnabled() || VI::EnableGuides->GetInt() == 1) &&
 		VI::SmoothSnap->GetInt() != 0 &&
 		!FMath::IsNearlyZero( SmoothSnapSpeed );
 
@@ -2473,23 +2574,40 @@ void UViewportWorldInteraction::PollInputIfNeeded()
 }
 
 
-void UViewportWorldInteraction::ResetTransformables()
-{
-	for( TUniquePtr<FViewportTransformable>& Transformable : Transformables )
-	{
-		Transformable->StartTransform = Transformable->LastTransform = Transformable->UnsnappedTargetTransform = Transformable->InterpolationSnapshotTransform = Transformable->TargetTransform = Transformable->GetTransform();
-	}
-}
-
-
-void UViewportWorldInteraction::RefreshTransformGizmo( const bool bNewObjectsSelected, bool bAllHandlesVisible, class UActorComponent* SingleVisibleHandle, const TArray< UActorComponent* >& HoveringOverHandles )
+void UViewportWorldInteraction::RefreshTransformGizmo( const bool bNewObjectsSelected )
 {
 	if ( IsTransformGizmoVisible() )
 	{
 		SpawnTransformGizmoIfNeeded();
 
 		// Make sure the gizmo is visible
-		TransformGizmoActor->SetIsTemporarilyHiddenInEditor( false );
+		const bool bShouldBeVisible = true;
+		const bool bPropagateToChildren = true;
+		TransformGizmoActor->GetRootComponent()->SetVisibility( bShouldBeVisible, bPropagateToChildren );
+
+
+		UViewportInteractor* DraggingWithInteractor = nullptr;
+		static TArray< UActorComponent* > HoveringOverHandles;
+		HoveringOverHandles.Reset();
+		for( UViewportInteractor* Interactor : Interactors )
+		{
+			FViewportInteractorData& InteractorData = Interactor->GetInteractorData();
+			if( InteractorData.DraggingMode == EViewportInteractionDraggingMode::TransformablesWithGizmo ||
+				InteractorData.DraggingMode == EViewportInteractionDraggingMode::TransformablesFreely ||
+				InteractorData.DraggingMode == EViewportInteractionDraggingMode::TransformablesAtLaserImpact )
+			{
+				DraggingWithInteractor = Interactor;
+			}
+
+			UActorComponent* HoveringOverHandle = InteractorData.HoveringOverTransformGizmoComponent.Get();
+			if( HoveringOverHandle != nullptr )
+			{
+				HoveringOverHandles.Add( HoveringOverHandle );
+			}
+		}
+		const bool bAllHandlesVisible = ( DraggingWithInteractor == nullptr );
+		UActorComponent* SingleVisibleHandle = ( DraggingWithInteractor != nullptr ) ? DraggingWithInteractor->GetInteractorData().DraggingTransformGizmoComponent.Get() : nullptr;
+
 
 		const ECoordSystem CurrentCoordSystem = GetTransformGizmoCoordinateSpace();
 		const bool bIsWorldSpaceGizmo = ( CurrentCoordSystem == COORD_World );
@@ -2509,20 +2627,13 @@ void UViewportWorldInteraction::RefreshTransformGizmo( const bool bNewObjectsSel
 		FBox GizmoSpaceSelectedObjectsBounds;
 		GizmoSpaceSelectedObjectsBounds.Init();
 		{
-			const FTransform WorldToGizmo = GizmoToWorld.Inverse();
+
 			for ( TUniquePtr<FViewportTransformable>& TransformablePtr : Transformables )
 			{
 				FViewportTransformable& Transformable = *TransformablePtr;
 
-				// Figure out the local space bounds of the entire object
-				const FBox LocalSpaceBounds = Transformable.GetLocalSpaceBoundingBox();
-
-				// Transform the bounds into the gizmo's space
-				const FTransform TransformableToWorld = Transformable.GetTransform();
-				const FTransform TransformableToGizmo = TransformableToWorld * WorldToGizmo;
-				const FBox GizmoSpaceBounds = LocalSpaceBounds.TransformBy( TransformableToGizmo );
-
-				// Get that bounding box into the local space of the gizmo
+				// Get the bounding box into the local space of the gizmo
+				const FBox GizmoSpaceBounds = Transformable.BuildBoundingBox( GizmoToWorld );
 				GizmoSpaceSelectedObjectsBounds += GizmoSpaceBounds;
 			}
 		}
@@ -2534,10 +2645,12 @@ void UViewportWorldInteraction::RefreshTransformGizmo( const bool bNewObjectsSel
 			Transformables[ 0 ]->IsUnorientedPoint();
 		const bool bAllowRotationAndScaleHandles = !bHaveSingleUnorientedPoint;
 
-		if ( VI::ForceGizmoPivotToCenterOfObjectsBounds->GetInt() > 0 )
+		if( ( Transformables.Num() > 1 && this->ViewportTransformer->ShouldCenterTransformGizmoPivot() ) ||
+			  VI::ForceGizmoPivotToCenterOfObjectsBounds->GetInt() > 0 )
 		{
-			GizmoToWorld.SetLocation( GizmoToWorld.TransformPosition( GizmoSpaceSelectedObjectsBounds.GetCenter() ) );
-			GizmoSpaceSelectedObjectsBounds = GizmoSpaceSelectedObjectsBounds.ShiftBy( -GizmoSpaceSelectedObjectsBounds.GetCenter() );
+			const FVector GizmoSpaceBoundsCenterLocation = GizmoSpaceSelectedObjectsBounds.GetCenter();
+			GizmoToWorld.SetLocation( GizmoToWorld.TransformPosition( GizmoSpaceBoundsCenterLocation ) );
+			GizmoSpaceSelectedObjectsBounds = GizmoSpaceSelectedObjectsBounds.ShiftBy( -GizmoSpaceBoundsCenterLocation );
 		}
 
 		if ( bNewObjectsSelected )
@@ -2581,7 +2694,9 @@ void UViewportWorldInteraction::RefreshTransformGizmo( const bool bNewObjectsSel
 		SpawnGridMeshActor();
 		if ( FSnappingUtils::IsSnapToGridEnabled() )
 		{
-			SnapGridActor->GetRootComponent()->SetVisibility( true );
+	        const bool bShouldBeVisible = true;
+	        const bool bPropagateToChildren = true;
+	        SnapGridActor->GetRootComponent()->SetVisibility( bShouldBeVisible, bPropagateToChildren );
 
 			const float GizmoAnimationAlpha = TransformGizmoActor->GetAnimationAlpha();
 
@@ -2596,7 +2711,10 @@ void UViewportWorldInteraction::RefreshTransformGizmo( const bool bNewObjectsSel
 			FBox GizmoLocalBoundsFlattened = GizmoLocalBounds;
 			GizmoLocalBoundsFlattened.Min.Z = GizmoLocalBoundsFlattened.Max.Z = 0.0f;
 			const FBox GizmoWorldBoundsFlattened = GizmoLocalBoundsFlattened.TransformBy( GizmoToWorld );
-			const float SnapGridSize = GizmoWorldBoundsFlattened.GetSize().GetAbsMax() * VI::SnapGridSize->GetFloat();
+
+			// Make sure we're at least as big as the gizmo bounds, but also large enough that you can see at least a few grid cells (depending on your snap size)
+			const float SnapGridSize = 
+				FMath::Max( GizmoWorldBoundsFlattened.GetSize().GetAbsMax(), GEditor->GetGridSize() * 2.5f ) * VI::SnapGridSize->GetFloat();
 
 			// The mesh is 100x100, so we'll scale appropriately
 			const float SnapGridScale = SnapGridSize / 100.0f;
@@ -2682,20 +2800,34 @@ void UViewportWorldInteraction::RefreshTransformGizmo( const bool bNewObjectsSel
 		else
 		{
 			// Grid snap not enabled
-			SnapGridActor->GetRootComponent()->SetVisibility( false );
+	        const bool bShouldBeVisible = false;
+	        const bool bPropagateToChildren = true;
+	        SnapGridActor->GetRootComponent()->SetVisibility( bShouldBeVisible, bPropagateToChildren );
 		}
+
+		if (bNewObjectsSelected && bPlayNextRefreshTransformGizmoSound)
+		{
+			PlaySound(AssetContainer->SelectionChangeSound, TransformGizmoActor->GetActorLocation());
+		}
+
+		bPlayNextRefreshTransformGizmoSound = true;
 	}
 	else
 	{
 		// Nothing selected or the gizmo was asked to be hidden
 		if( TransformGizmoActor != nullptr )
 		{
-			TransformGizmoActor->SetIsTemporarilyHiddenInEditor( true );
+	        const bool bShouldBeVisible = false;
+	        const bool bPropagateToChildren = true;
+	        TransformGizmoActor->GetRootComponent()->SetVisibility( bShouldBeVisible, bPropagateToChildren );
 		}
 		GizmoLocalBounds = FBox( 0 );
 
 		// Hide the snap actor
-		GetSnapGridActor()->GetRootComponent()->SetVisibility( false );
+		if( SnapGridActor != nullptr )
+		{
+			SnapGridActor->GetRootComponent()->SetVisibility( false );
+		}
 	}
 }
 
@@ -2718,14 +2850,16 @@ void UViewportWorldInteraction::SpawnTransformGizmoIfNeeded()
 
 		// Create the correct gizmo
 		const bool bWithSceneComponent = false;	// We already have our own scene component
-		TransformGizmoActor = CastChecked<ABaseTransformGizmo>( SpawnTransientSceneActor( GetWorld(), TransformGizmoClass, TEXT( "PivotTransformGizmo" ), bWithSceneComponent ) );
+		TransformGizmoActor = CastChecked<ABaseTransformGizmo>(SpawnTransientSceneActor(TransformGizmoClass, TEXT( "PivotTransformGizmo" ), bWithSceneComponent));
 
 		check( TransformGizmoActor != nullptr );
 		TransformGizmoActor->SetOwnerWorldInteraction( this );
 
 		if ( !IsTransformGizmoVisible() )
 		{
-			TransformGizmoActor->SetIsTemporarilyHiddenInEditor( true );
+	        const bool bShouldBeVisible = false;
+	        const bool bPropagateToChildren = true;
+	        TransformGizmoActor->GetRootComponent()->SetVisibility( bShouldBeVisible, bPropagateToChildren );
 		}
 	}
 }
@@ -2744,7 +2878,7 @@ bool UViewportWorldInteraction::ShouldTransformGizmoBeVisible() const
 
 bool UViewportWorldInteraction::IsTransformGizmoVisible() const
 {
-	return ( Transformables.Num() > 0 && VI::ShowTransformGizmo->GetInt() != 0 && bShouldTransformGizmoBeVisible );
+	return ( Transformables.Num() > 0 && VI::ShowTransformGizmo->GetInt() != 0 && bShouldTransformGizmoBeVisible && !HasTransformableWithVelocityInSimulate() );
 }
 
 void UViewportWorldInteraction::SetTransformGizmoScale( const float NewScale )
@@ -2814,7 +2948,7 @@ float UViewportWorldInteraction::GetMinScale()
 	return VI::ScaleMin->GetFloat();
 }
 
-void UViewportWorldInteraction::SetWorldToMetersScale( const float NewWorldToMetersScale )
+void UViewportWorldInteraction::SetWorldToMetersScale( const float NewWorldToMetersScale, const bool bCompensateRoomWorldScale  /*= false*/ )
 {
 	// @todo vreditor: This is bad because we're clobbering the world settings which will be saved with the map.  Instead we need to 
 	// be able to apply an override before the scene view gets it
@@ -2822,6 +2956,14 @@ void UViewportWorldInteraction::SetWorldToMetersScale( const float NewWorldToMet
 	ENGINE_API extern float GNewWorldToMetersScale;
 	GNewWorldToMetersScale = NewWorldToMetersScale;
 	OnWorldScaleChangedEvent.Broadcast(NewWorldToMetersScale);
+
+	if (bCompensateRoomWorldScale)
+	{
+		const FVector RoomspacePivotLocation = GetRoomSpaceHeadTransform().GetLocation();
+		FTransform NewRoomTransform = GetRoomTransform();
+		CompensateRoomTransformForWorldScale(NewRoomTransform, NewWorldToMetersScale, RoomspacePivotLocation);
+		RoomTransformToSetOnFrame = MakeTuple( NewRoomTransform, CurrentTickNumber + 1 );
+	}
 }
 
 UViewportInteractor* UViewportWorldInteraction::GetOtherInteractorIntertiaContribute( UViewportInteractor* Interactor )
@@ -2849,54 +2991,17 @@ UViewportInteractor* UViewportWorldInteraction::GetOtherInteractorIntertiaContri
 	return OtherInteractorThatWasAssistingDrag;
 }
 
-AActor* UViewportWorldInteraction::SpawnTransientSceneActor( UWorld* World, TSubclassOf<AActor> ActorClass, const FString& ActorName, const bool bWithSceneComponent )
-{
-	check( World != nullptr );
-	const bool bWasWorldPackageDirty = World->GetOutermost()->IsDirty();
-
-	FActorSpawnParameters ActorSpawnParameters;
-	ActorSpawnParameters.Name = MakeUniqueObjectName( World, ActorClass, *ActorName );
-	ActorSpawnParameters.SpawnCollisionHandlingOverride = ESpawnActorCollisionHandlingMethod::AlwaysSpawn;
-
-	if( VI::UseTransientActors->GetInt() != 0 )
-	{
-		ActorSpawnParameters.ObjectFlags = EObjectFlags::RF_Transient;
-	}
-
-	check( ActorClass != nullptr );
-	AActor* NewActor = World->SpawnActor< AActor >( ActorClass, ActorSpawnParameters );
-	NewActor->SetActorLabel( ActorName );
-
-	if ( bWithSceneComponent )
-	{
-		// Give the new actor a root scene component, so we can attach multiple sibling components to it
-		USceneComponent* SceneComponent = NewObject<USceneComponent>( NewActor );
-		NewActor->AddOwnedComponent( SceneComponent );
-		NewActor->SetRootComponent( SceneComponent );
-		SceneComponent->RegisterComponent();
-	}
-
-	// Don't dirty the level file after spawning a transient actor
-	if ( !bWasWorldPackageDirty )
-	{
-		World->GetOutermost()->SetDirtyFlag( false );
-	}
-
-	return NewActor;
-}
-
-
 void UViewportWorldInteraction::DestroyActors()
 {
 	if(TransformGizmoActor != nullptr)
 	{
-		DestroyTransientActor( GetWorld(), TransformGizmoActor );
+		DestroyTransientActor(TransformGizmoActor);
 		TransformGizmoActor = nullptr;
 	}
 
 	if(SnapGridActor != nullptr)
 	{
-		DestroyTransientActor( GetWorld(), SnapGridActor );
+		DestroyTransientActor(SnapGridActor);
 		SnapGridActor = nullptr;
 		SnapGridMeshComponent = nullptr;
 	}
@@ -2908,26 +3013,94 @@ void UViewportWorldInteraction::DestroyActors()
 	}
 }
 
-void UViewportWorldInteraction::DestroyTransientActor( UWorld* World, AActor* Actor )
+bool UViewportWorldInteraction::AreAligningToActors()
 {
-	check( World != nullptr );
-	if ( Actor != nullptr )
+	return (VI::EnableGuides->GetInt() == 1) ? true : false;
+}
+
+bool UViewportWorldInteraction::HasCandidatesSelected()
+{
+	return CandidateActors.Num() > 0 ? true : false;
+}
+
+void UViewportWorldInteraction::SetSelectionAsCandidates()
+{
+	if (HasCandidatesSelected())
 	{
-		const bool bWasWorldPackageDirty = World->GetOutermost()->IsDirty();
+		CandidateActors.Reset();
+	}
+	else if (VI::EnableGuides->GetInt() == 1)
+	{
+		TArray<TUniquePtr<FViewportTransformable>> NewTransformables;
+
+		USelection* ActorSelectionSet = GEditor->GetSelectedActors();
 
-		const bool bNetForce = false;
-		const bool bShouldModifyLevel = false;	// Don't modify level for transient actor destruction
-		World->DestroyActor(Actor, bNetForce, bShouldModifyLevel);
-		Actor = nullptr;
+		static TArray<UObject*> SelectedActorObjects;
+		SelectedActorObjects.Reset();
+		ActorSelectionSet->GetSelectedObjects(AActor::StaticClass(), /* Out */ SelectedActorObjects);
 
-		// Don't dirty the level file after destroying a transient actor
-		if (!bWasWorldPackageDirty)
+		for (UObject* SelectedActorObject : SelectedActorObjects)
 		{
-			World->GetOutermost()->SetDirtyFlag(false);
+			AActor* SelectedActor = CastChecked<AActor>(SelectedActorObject);
+			CandidateActors.Add(SelectedActor);
 		}
+
+		Deselect();
+	}
+}
+
+FVector UViewportWorldInteraction::GetLocalIntersectPointOnRotationGizmo() const
+{
+	return LocalIntersectPointOnRotationGizmo;
+}
+
+float UViewportWorldInteraction::GetCurrentDeltaTime() const
+{
+	return CurrentDeltaTime;
+}
+
+bool UViewportWorldInteraction::ShouldSuppressExistingCursor() const
+{
+	if (VI::ForceShowCursor->GetInt() == 1)
+	{
+		return false;
+	}
+	else
+	{
+		return bShouldSuppressCursor;
+	}
+
+}
+
+const UViewportInteractionAssetContainer& UViewportWorldInteraction::GetAssetContainer() const
+{
+	return *AssetContainer;
+}
+
+const class UViewportInteractionAssetContainer& UViewportWorldInteraction::LoadAssetContainer()
+{
+	return *LoadObject<UViewportInteractionAssetContainer>(nullptr, *UViewportWorldInteraction::AssetContainerPath);
+}
+
+void UViewportWorldInteraction::PlaySound(USoundBase* SoundBase, const FVector& InWorldLocation, const float InVolume /*= 1.0f*/)
+{
+	if (bActive)
+	{
+		const float Volume = InVolume*VI::SFXMultiplier->GetFloat();
+		UGameplayStatics::PlaySoundAtLocation(GetWorld(), SoundBase, InWorldLocation, FRotator::ZeroRotator, Volume);
 	}
 }
 
+void UViewportWorldInteraction::SetInVR(const bool bInVR)
+{
+	bIsInVR = bInVR;
+}
+
+bool UViewportWorldInteraction::IsInVR() const
+{
+	return bIsInVR;
+}
+
 EGizmoHandleTypes UViewportWorldInteraction::GetCurrentGizmoType() const
 {
 	switch( GLevelEditorModeTools().GetWidgetMode() )
@@ -2978,8 +3151,9 @@ void UViewportWorldInteraction::SetTransformGizmoClass( const TSubclassOf<ABaseT
 	TransformGizmoClass = NewTransformGizmoClass;
 }
 
-void UViewportWorldInteraction::SetDraggedInteractable( IViewportInteractableInterface* InDraggedInteractable )
+void UViewportWorldInteraction::SetDraggedInteractable( IViewportInteractableInterface* InDraggedInteractable, UViewportInteractor* Interactor )
 {
+	Interactor->SetDraggingMode(EViewportInteractionDraggingMode::Interactable);
 	DraggedInteractable = InDraggedInteractable;
 	if ( DraggedInteractable && DraggedInteractable->GetDragOperationComponent() )
 	{
@@ -2995,7 +3169,7 @@ bool UViewportWorldInteraction::IsOtherInteractorHoveringOverComponent( UViewpor
 	{
 		for ( UViewportInteractor* CurrentInteractor : Interactors )
 		{
-			if ( CurrentInteractor != Interactor && CurrentInteractor->GetHoverComponent() == Component )
+			if ( CurrentInteractor != Interactor && CurrentInteractor->GetLastHoverComponent() == Component )
 			{
 				bResult = true;
 				break;
@@ -3012,20 +3186,21 @@ void UViewportWorldInteraction::SpawnGridMeshActor()
 	if( SnapGridActor == nullptr )
 	{
 		const bool bWithSceneComponent = false;
-		SnapGridActor = SpawnTransientSceneActor<AActor>( GetWorld(), TEXT( "SnapGrid" ), bWithSceneComponent );
+		SnapGridActor = SpawnTransientSceneActor<AActor>(TEXT("SnapGrid"), bWithSceneComponent);
 
 		SnapGridMeshComponent = NewObject<UStaticMeshComponent>( SnapGridActor );
+		SnapGridMeshComponent->MarkAsEditorOnlySubobject();
 		SnapGridActor->AddOwnedComponent( SnapGridMeshComponent );
 		SnapGridActor->SetRootComponent( SnapGridMeshComponent );
 		SnapGridMeshComponent->RegisterComponent();
 
-		UStaticMesh* GridMesh = LoadObject<UStaticMesh>( nullptr, TEXT( "/Engine/VREditor/SnapGrid/SnapGridPlaneMesh" ) );
+		UStaticMesh* GridMesh = AssetContainer->GridMesh;
 		check( GridMesh != nullptr );
 		SnapGridMeshComponent->SetStaticMesh( GridMesh );
 		SnapGridMeshComponent->SetMobility( EComponentMobility::Movable );
 		SnapGridMeshComponent->SetCollisionEnabled( ECollisionEnabled::NoCollision );
 
-		UMaterialInterface* GridMaterial = LoadObject<UMaterialInterface>( nullptr, TEXT( "/Engine/VREditor/SnapGrid/SnapGridMaterial" ) );
+		UMaterialInterface* GridMaterial = AssetContainer->GridMaterial;
 		check( GridMaterial != nullptr );
 
 		SnapGridMID = UMaterialInstanceDynamic::Create( GridMaterial, GetTransientPackage() );
@@ -3043,7 +3218,7 @@ FVector UViewportWorldInteraction::CalculateAverageLocationOfTransformables()
 
 	for( TUniquePtr<FViewportTransformable>& TransformablePtr : Transformables )
 	{
-		Result += ((*TransformablePtr).LastTransform.GetLocation());
+		Result += TransformablePtr.Get()->GetTransform().GetLocation();
 	}
 
 	Result /= Transformables.Num();
@@ -3051,11 +3226,12 @@ FVector UViewportWorldInteraction::CalculateAverageLocationOfTransformables()
 	return Result;
 }
 
-FLinearColor UViewportWorldInteraction::GetColor( const EColors Color ) const
+FLinearColor UViewportWorldInteraction::GetColor(const EColors Color, const float Multiplier /*= 1.f*/) const
 {
-	return Colors[ (int32)Color ];
+	return Colors[ (int32)Color ] * Multiplier;
 }
 
+
 void UViewportWorldInteraction::AddMouseCursorInteractor()
 {
 	if( ++DefaultMouseCursorInteractorRefCount == 1 )
@@ -3084,5 +3260,380 @@ void UViewportWorldInteraction::ReleaseMouseCursorInteractor()
 	}
 }
 
+FVector UViewportWorldInteraction::FindTransformGizmoAlignPoint(const FTransform& GizmoStartTransform, const FTransform& DesiredGizmoTransform, const bool bShouldConstrainMovement, FVector ConstraintAxes)
+{
+	struct Local
+	{
+		static TArray<FVector> FindLocalSnapPoints(const FBox InBox)
+		{
+			TArray<FVector> PotentialSnapPoints;
+			FVector BoxCenter;
+			FVector BoxExtents;
+
+			InBox.GetCenterAndExtents(BoxCenter, BoxExtents);
+
+			FVector PotentialSnapPoint = FVector::ZeroVector;
+
+			// Potential snap points are:
+			// The center of each face
+			for (int32 X = -1; X < 2; ++X)
+			{
+				PotentialSnapPoint[0] = X*BoxExtents[0];
+				PotentialSnapPoints.AddUnique(PotentialSnapPoint);
+				// The center of each edge
+				for (int32 Y = -1; Y < 2; ++Y)
+				{
+					PotentialSnapPoint[1] = Y*BoxExtents[1];
+					PotentialSnapPoints.AddUnique(PotentialSnapPoint);
+					// Each corner
+					for (int32 Z = -1; Z < 2; ++Z)
+					{
+						PotentialSnapPoint[2] = Z*BoxExtents[2];
+						PotentialSnapPoints.AddUnique(PotentialSnapPoint);
+					}
+				}
+			}
+			return PotentialSnapPoints;
+		}
+	};
+
+	TArray<FGuideData> PotentialGizmoGuides;
+	// Get all the potential snap points on the transform gizmo at the desired location
+	const TArray<FVector> DesiredGizmoLocalGizmoSnapPoints = Local::FindLocalSnapPoints(GizmoLocalBounds);
+	DrawBoxBrackets(GizmoLocalBounds, DesiredGizmoTransform, FLinearColor::Yellow);
+
+	// Don't let the guide lines be shorter than the local bounding box extent in any direction. 
+	// This helps when objects are close together
+	FVector GizmoLocalBoundsExtents;
+	FVector GizmoLocalBoundsCenter;
+	GizmoLocalBounds.GetCenterAndExtents(GizmoLocalBoundsCenter, GizmoLocalBoundsExtents);
+	const FVector MinGuideLength = GizmoLocalBoundsExtents - GizmoLocalBoundsCenter;
+	
+	if(bShouldConstrainMovement)
+	{
+		ConstraintAxes = GizmoStartTransform.InverseTransformVector(ConstraintAxes);
+	}
+
+	// Our snap distances are some percentage of the transform gizmo's dimensions
+	const float AdjustedSnapDistance = (VI::ForceSnapDistance->GetFloat() / 100.0f) * 2.0f * (MinGuideLength.GetAbsMax());
+	int32 NumberOfMatchesNeeded = 0;
+	if (bShouldConstrainMovement)
+	{
+		for (int32 PointAxis = 0; PointAxis < 3; ++PointAxis)
+		{
+			if (!FMath::IsNearlyZero(ConstraintAxes[PointAxis], 0.0001f))
+			{
+				NumberOfMatchesNeeded++;
+			}
+		}
+	}
+	else
+	{
+		NumberOfMatchesNeeded = 3;
+	}
+
+	TArray<const AActor*> UsingCandidateActors;
+	if (HasCandidatesSelected())
+	{	
+		for (const AActor* SelectedCandidateActor : CandidateActors)
+		{
+			// Don't align to yourself, the entire world, or any actors hidden in the editor
+			if (!SelectedCandidateActor->IsSelected()
+				&& SelectedCandidateActor != GetWorld()->GetDefaultBrush()
+				&& SelectedCandidateActor->IsHiddenEd() == false
+				&& !SelectedCandidateActor->IsEditorOnly()
+				&& SelectedCandidateActor->GetRootComponent() != nullptr
+				&& !SelectedCandidateActor->GetRootComponent()->IsEditorOnly())
+			{
+				UsingCandidateActors.Add(SelectedCandidateActor);
+				const FBox LocalActorBoundingBox = SelectedCandidateActor->CalculateComponentsBoundingBoxInLocalSpace();
+				DrawBoxBrackets(LocalActorBoundingBox, SelectedCandidateActor->GetTransform(), FLinearColor::Blue);
+			}
+		}
+	}
+	else
+	{
+		// Find all possible candidates for alignment
+		// TODO: add the world grid
+		// TODO: remove anything it might not make sense to align to
+		const float CompareDistance = VI::AlignCandidateDistance->GetFloat() * MinGuideLength.GetAbsMax();
+		const FVector Start = DesiredGizmoTransform.GetLocation();
+		const FVector End = DesiredGizmoTransform.GetLocation();
+		TArray<FOverlapResult> OutOverlaps;
+		bool const bHit = GetWorld()->OverlapMultiByChannel(OutOverlaps, Start, FQuat::Identity, ECC_Visibility, FCollisionShape::MakeSphere(CompareDistance));
+
+		if (bHit)
+		{
+			for (FOverlapResult OverlapResult : OutOverlaps)
+			{
+				const AActor* PossibleCandidateActor = OverlapResult.GetActor();
+
+				// Don't align to yourself, the entire world, or any actors hidden in the editor
+				if (!PossibleCandidateActor->IsSelected()
+					&& PossibleCandidateActor != GetWorld()->GetDefaultBrush()
+					&& PossibleCandidateActor->IsHiddenEd() == false
+					&& !PossibleCandidateActor->IsEditorOnly()
+					&& PossibleCandidateActor->GetRootComponent() != nullptr
+					&& !PossibleCandidateActor->GetRootComponent()->IsEditorOnly())
+				{
+					{
+						// Check if our candidate actor is close enough (comparison multiplier * maximum dimension of actor)
+						if (
+							(DesiredGizmoTransform.GetLocation() - PossibleCandidateActor->GetActorLocation()).GetAbsMin() <= CompareDistance)
+						{
+							UsingCandidateActors.Add(PossibleCandidateActor);
+						}
+					}
+				}
+			}
+		}
+	}
+
+	for (const AActor* CandidateActor : UsingCandidateActors)
+	{
+		// Get the gizmo space snap points for the stationary candidate actor
+		const FBox LocalActorBox = CandidateActor->CalculateComponentsBoundingBoxInLocalSpace();
+		const TArray<FVector> LocalCandidateSnapPoints = Local::FindLocalSnapPoints(LocalActorBox);
+
+		// Set up the initial guide information for X, Y, and Z guide line
+		FGuideData InitialGuideHelper;
+		InitialGuideHelper.AlignedActor = CandidateActor;
+		InitialGuideHelper.LocalOffset = InitialGuideHelper.SnapPoint = InitialGuideHelper.GuideStart = InitialGuideHelper.GuideEnd = FVector::ZeroVector;
+		InitialGuideHelper.GuideColor = FLinearColor::Black.ToFColor(/*bSRGB=*/ true);
+		InitialGuideHelper.DrawAlpha = 1.0f;
+		InitialGuideHelper.GuideLength = 10000000.0f;
+
+		for (const FVector LocalCandidateSnapPoint : LocalCandidateSnapPoints)
+		{
+			FVector WorldCandidateSnapPoint = CandidateActor->ActorToWorld().TransformPosition(LocalCandidateSnapPoint);
+			FVector DesiredGizmoLocalCandidateSnapPoint = DesiredGizmoTransform.InverseTransformPosition(WorldCandidateSnapPoint);
+			// Check it against each moving snap point
+			for (const FVector DesiredGizmoLocalGizmoSnapPoint : DesiredGizmoLocalGizmoSnapPoints)
+			{
+				FVector WorldGizmoSnapPoint = DesiredGizmoTransform.TransformPosition(DesiredGizmoLocalGizmoSnapPoint);
+				int32 NumberOfMatchingAxes = 0;
+				FVector DesiredGizmoLocalOffset = FVector::ZeroVector;
+				for (int32 PointAxis = 0; PointAxis < 3; ++PointAxis)
+				{
+					// If we are within the snap distance and can snap along that axis
+					if (FMath::Abs(DesiredGizmoLocalCandidateSnapPoint[PointAxis] - DesiredGizmoLocalGizmoSnapPoint[PointAxis]) <= AdjustedSnapDistance &&
+						(!FMath::IsNearlyZero(ConstraintAxes[PointAxis], 0.0001f) ||
+							!bShouldConstrainMovement ))
+					{
+						NumberOfMatchingAxes++;
+						DesiredGizmoLocalOffset[PointAxis] = DesiredGizmoLocalCandidateSnapPoint[PointAxis] - DesiredGizmoLocalGizmoSnapPoint[PointAxis];
+					}
+				}
+
+				if (NumberOfMatchingAxes >= NumberOfMatchesNeeded)
+				{
+					FVector DesiredGizmoLocalGuideStart = DesiredGizmoLocalGizmoSnapPoint + DesiredGizmoLocalOffset;
+
+					// Transform the goal location into world space
+					const FVector WorldGuideStart = DesiredGizmoTransform.TransformPosition( DesiredGizmoLocalGuideStart );
+
+					InitialGuideHelper.LocalOffset = DesiredGizmoLocalOffset;
+					InitialGuideHelper.GuideLength = (WorldCandidateSnapPoint - WorldGuideStart).Size();
+					InitialGuideHelper.GuideStart = WorldGuideStart;
+					InitialGuideHelper.GuideEnd = WorldCandidateSnapPoint;
+					InitialGuideHelper.SnapPoint = WorldGizmoSnapPoint;
+					InitialGuideHelper.GuideColor = FLinearColor::Yellow.ToFColor(/*bSRGB=*/ true);
+					PotentialGizmoGuides.Add(InitialGuideHelper);	
+				}
+			}
+		}
+	}
+
+	// Now find the best guide from all available point combinations
+	FGuideData AlignedGuide;
+	bool bFoundAlignedTransform = false;
+	float AlignedDeltaSize = 10000000.0f;
+	for (FGuideData PotentialGizmoGuide : PotentialGizmoGuides)
+	{	
+		const float OffsetSize = PotentialGizmoGuide.LocalOffset.Size();
+
+		// Keep finding the guide with the shortest offset size
+		if (OffsetSize < AlignedDeltaSize && !FMath::IsNearlyZero(OffsetSize))
+		{
+			bFoundAlignedTransform = true;
+			AlignedGuide = PotentialGizmoGuide;
+			AlignedDeltaSize = OffsetSize;
+		}
+	}
+
+	FVector LastBestAlignedLocationOffset = FVector::ZeroVector;
+	if (bFoundAlignedTransform)
+	{
+		LastBestAlignedLocationOffset = GizmoStartTransform.InverseTransformPosition(AlignedGuide.GuideStart) - GizmoStartTransform.InverseTransformPosition(AlignedGuide.SnapPoint);
+	}
+
+	return LastBestAlignedLocationOffset;
+}
+
+void UViewportWorldInteraction::AddActorToExcludeFromHitTests( AActor* ActorToExcludeFromHitTests )
+{
+	ActorsToExcludeFromHitTest.Add( ActorToExcludeFromHitTests );
+
+	// Remove expired entries
+	for( int32 ActorIndex = 0; ActorIndex < ActorsToExcludeFromHitTest.Num(); ++ActorIndex )
+	{
+		if( !ActorsToExcludeFromHitTest[ ActorIndex ].IsValid() )
+		{
+			ActorsToExcludeFromHitTest.RemoveAtSwap( ActorIndex-- );
+		}
+	}
+}
+
+
+
+
+void UViewportWorldInteraction::DrawBoxBrackets(const FBox InActor, const FTransform LocalToWorld, const FLinearColor BracketColor)
+{
+	struct Local
+	{
+		static void GetBoundingVectors(const FBox LocalBox, FVector& OutVectorMin, FVector& OutVectorMax)
+		{
+			OutVectorMin = FVector(BIG_NUMBER);
+			OutVectorMax = FVector(-BIG_NUMBER);
+
+
+			// MinVector
+			OutVectorMin.X = FMath::Min<float>(LocalBox.Min.X, OutVectorMin.X);
+			OutVectorMin.Y = FMath::Min<float>(LocalBox.Min.Y, OutVectorMin.Y);
+			OutVectorMin.Z = FMath::Min<float>(LocalBox.Min.Z, OutVectorMin.Z);
+			// MaxVector
+			OutVectorMax.X = FMath::Max<float>(LocalBox.Max.X, OutVectorMax.X);
+			OutVectorMax.Y = FMath::Max<float>(LocalBox.Max.Y, OutVectorMax.Y);
+			OutVectorMax.Z = FMath::Max<float>(LocalBox.Max.Z, OutVectorMax.Z);
+		}
+	};
+
+	const FColor GROUP_COLOR = BracketColor.ToFColor(/*bSRGB=*/ true);
+
+	FVector MinVector;
+	FVector MaxVector;
+	Local::GetBoundingVectors(InActor, MinVector, MaxVector);
+
+	// Create a bracket offset to determine the length of our corner axises
+	const float BracketOffset = FVector::Dist(MinVector, MaxVector) * 0.1f;
+
+	// Calculate bracket corners based on min/max vectors
+	TArray<FVector> BracketCorners;
+
+	// Bottom Corners
+	BracketCorners.Add(FVector(MinVector.X, MinVector.Y, MinVector.Z));
+	BracketCorners.Add(FVector(MinVector.X, MaxVector.Y, MinVector.Z));
+	BracketCorners.Add(FVector(MaxVector.X, MaxVector.Y, MinVector.Z));
+	BracketCorners.Add(FVector(MaxVector.X, MinVector.Y, MinVector.Z));
+
+	// Top Corners
+	BracketCorners.Add(FVector(MinVector.X, MinVector.Y, MaxVector.Z));
+	BracketCorners.Add(FVector(MinVector.X, MaxVector.Y, MaxVector.Z));
+	BracketCorners.Add(FVector(MaxVector.X, MaxVector.Y, MaxVector.Z));
+	BracketCorners.Add(FVector(MaxVector.X, MinVector.Y, MaxVector.Z));
+
+	for (int32 BracketCornerIndex = 0; BracketCornerIndex < BracketCorners.Num(); ++BracketCornerIndex)
+	{
+		// Direction corner axis should be pointing based on min/max
+		const FVector CORNER = BracketCorners[BracketCornerIndex];
+		const int32 DIR_X = CORNER.X == MaxVector.X ? -1 : 1;
+		const int32 DIR_Y = CORNER.Y == MaxVector.Y ? -1 : 1;
+		const int32 DIR_Z = CORNER.Z == MaxVector.Z ? -1 : 1;
+
+		const FVector LocalBracketX = FVector(CORNER.X + (BracketOffset * DIR_X), CORNER.Y, CORNER.Z);
+		const FVector LocalBracketY = FVector(CORNER.X, CORNER.Y + (BracketOffset * DIR_Y), CORNER.Z);
+		const FVector LocalBracketZ = FVector(CORNER.X, CORNER.Y, CORNER.Z + (BracketOffset * DIR_Z));
+
+		const FVector WorldCorner = LocalToWorld.TransformPosition(CORNER);
+		const FVector WorldBracketX = LocalToWorld.TransformPosition(LocalBracketX);
+		const FVector WorldBracketY = LocalToWorld.TransformPosition(LocalBracketY);
+		const FVector WorldBracketZ = LocalToWorld.TransformPosition(LocalBracketZ);
+
+		DrawDebugLine(GetWorld(), WorldCorner, WorldBracketX, GROUP_COLOR, false, -1.0f, 1.0f, 2.0f);
+		DrawDebugLine(GetWorld(), WorldCorner, WorldBracketY, GROUP_COLOR, false, -1.0f, 1.0f, 2.0f);
+		DrawDebugLine(GetWorld(), WorldCorner, WorldBracketZ, GROUP_COLOR, false, -1.0f, 1.0f, 2.0f);
+	}
+
+}
+
+void UViewportWorldInteraction::CompensateRoomTransformForWorldScale(FTransform& InOutRoomTransform, const float InNewWorldToMetersScale, const FVector& InRoomPivotLocation)
+{
+	const float OldWorldToMetersScale = GetWorld()->GetWorldSettings()->WorldToMeters;
+
+	// Because the tracking space size has changed, but our head position within that space relative to the origin
+	// of the room is the same (before scaling), we need to offset our location within the tracking space to compensate.
+	// This makes the user feel like their head and hands remain in the same location.
+	const FVector WorldSpacePivotLocation = InOutRoomTransform.TransformPosition(InRoomPivotLocation);
+	const FVector NewRoomSpacePivotLocation = (InRoomPivotLocation / OldWorldToMetersScale) * InNewWorldToMetersScale;
+	const FVector NewWorldSpacePivotLocation = InOutRoomTransform.TransformPosition(NewRoomSpacePivotLocation);
+	const FVector WorldSpacePivotDelta = (NewWorldSpacePivotLocation - WorldSpacePivotLocation);
+	const FVector NewWorldSpaceRoomLocation = InOutRoomTransform.GetLocation() - WorldSpacePivotDelta;
+
+	InOutRoomTransform.SetLocation(NewWorldSpaceRoomLocation);
+}
+
+bool UViewportWorldInteraction::HasTransformableWithVelocityInSimulate() const
+{
+	bool bResult = false;
+
+	// Only check if we are in simulate and if the world of this extension is actually the world simulating.
+	if (GEditor->bIsSimulatingInEditor && GetWorld() == GEditor->PlayWorld)
+	{
+		for (const TUniquePtr<FViewportTransformable>& TransformablePtr : Transformables)
+		{
+			const FViewportTransformable& Transformable = *TransformablePtr;
+
+			if (!Transformable.GetLinearVelocity().IsNearlyZero(1.0f))
+			{
+				bResult = true;
+				break;
+			}
+		}
+	}
+
+	return bResult;
+}
+
+FVector UViewportWorldInteraction::SnapLocation(const bool bLocalSpaceSnapping, const FVector& DesiredGizmoLocation, const FTransform &GizmoStartTransform, const FVector SnapGridBase, const bool bShouldConstrainMovement, const FVector AlignAxes)
+{
+	bool bTransformableAlignmentUsed = false;
+	FVector SnappedGizmoLocation = DesiredGizmoLocation;
+
+	const FVector GizmoSpaceDesiredGizmoLocation = GizmoStartTransform.InverseTransformPosition( DesiredGizmoLocation );
+	
+	if ((VI::EnableGuides->GetInt() == 1) && bLocalSpaceSnapping && ViewportTransformer->CanAlignToActors() == true)
+	{
+		FTransform DesiredGizmoTransform = GizmoStartTransform;
+		DesiredGizmoTransform.SetLocation( DesiredGizmoLocation );
+
+		FVector LocationOffset = FVector::ZeroVector;
+		LocationOffset = FindTransformGizmoAlignPoint(GizmoStartTransform, DesiredGizmoTransform, bShouldConstrainMovement, AlignAxes);
+		
+		if (!LocationOffset.IsZero())
+		{
+			bTransformableAlignmentUsed = true;
+			const FVector GizmoSpaceSnappedGizmoLocation = GizmoSpaceDesiredGizmoLocation + LocationOffset;
+			SnappedGizmoLocation = GizmoStartTransform.TransformPosition( GizmoSpaceSnappedGizmoLocation );
+		}
+	}
+
+	if (FSnappingUtils::IsSnapToGridEnabled() && !bTransformableAlignmentUsed)
+	{
+		// Snap in local space, if we need to
+		if (bLocalSpaceSnapping)
+		{
+			FVector GizmoSpaceSnappedGizmoLocation = GizmoSpaceDesiredGizmoLocation;
+			FSnappingUtils::SnapPointToGrid(GizmoSpaceSnappedGizmoLocation, SnapGridBase);
+			SnappedGizmoLocation = GizmoStartTransform.TransformPosition(GizmoSpaceSnappedGizmoLocation);
+		}
+		else
+		{
+			FSnappingUtils::SnapPointToGrid(SnappedGizmoLocation, SnapGridBase);
+		}
+	}	
+
+	return SnappedGizmoLocation;
+}
+
+#undef LOCTEXT_NAMESPACE
 
-#undef LOCTEXT_NAMESPACE
\ No newline at end of file
