Manual copy of 3341264
- Use optional member to store gizmo handle type for Universal gizmo.

diff --git a/Engine/Source/Editor/ViewportInteraction/Public/ViewportWorldInteraction.h b/Engine/Source/Editor/ViewportInteraction/Public/ViewportWorldInteraction.h
--- a/Engine/Source/Editor/ViewportInteraction/Public/ViewportWorldInteraction.h
+++ b/Engine/Source/Editor/ViewportInteraction/Public/ViewportWorldInteraction.h
@@ -606,6 +606,9 @@ private:
 	/** All of the objects we're currently interacting with, such as selected actors */
 	TArray< TUniquePtr< FViewportTransformable > > Transformables;
 
+	/** The current gizmo type */ //@todo ViewportInteraction: Currently this is only used for universal gizmo.
+	TOptional<EGizmoHandleTypes> GizmoType;
+
 	//
 	// Snap grid
 	//
diff --git a/Engine/Source/Editor/ViewportInteraction/ViewportWorldInteraction.cpp b/Engine/Source/Editor/ViewportInteraction/ViewportWorldInteraction.cpp
--- a/Engine/Source/Editor/ViewportInteraction/ViewportWorldInteraction.cpp
+++ b/Engine/Source/Editor/ViewportInteraction/ViewportWorldInteraction.cpp
@@ -279,6 +279,7 @@ UViewportWorldInteraction::UViewportWorldInteraction():
 	DraggingRotationHandleDirection(),
 	bShouldTransformGizmoBeVisible( true ),
 	TransformGizmoScale( VI::GizmoScaleInDesktop->GetFloat() ),
+	GizmoType(),
 	SnapGridActor( nullptr ),
 	SnapGridMeshComponent( nullptr ),
 	SnapGridMID( nullptr ),
@@ -392,6 +393,8 @@ void UViewportWorldInteraction::Shutdown()
 
 	AssetContainer = nullptr;
 
+	GizmoType.Reset();
+
 	USelection::SelectionChangedEvent.RemoveAll( this );
 	GEditor->OnEditorClose().RemoveAll( this );
 }
@@ -3097,19 +3100,26 @@ bool UViewportWorldInteraction::IsInVR() const
 
 EGizmoHandleTypes UViewportWorldInteraction::GetCurrentGizmoType() const
 {
-	switch( GLevelEditorModeTools().GetWidgetMode() )
+	if (GizmoType.IsSet())
+	{
+		return GizmoType.GetValue();
+	}
+	else
 	{
-		case FWidget::WM_TranslateRotateZ:
-			return EGizmoHandleTypes::All;
+		switch( GLevelEditorModeTools().GetWidgetMode() )
+		{
+			case FWidget::WM_TranslateRotateZ:
+				return EGizmoHandleTypes::All;
 
-		case FWidget::WM_Translate:
-			return EGizmoHandleTypes::Translate;
+			case FWidget::WM_Translate:
+				return EGizmoHandleTypes::Translate;
 
-		case FWidget::WM_Rotate:
-			return EGizmoHandleTypes::Rotate;
+			case FWidget::WM_Rotate:
+				return EGizmoHandleTypes::Rotate;
 
-		case FWidget::WM_Scale:
-			return EGizmoHandleTypes::Scale;
+			case FWidget::WM_Scale:
+				return EGizmoHandleTypes::Scale;
+		}
 	}
 
 	return EGizmoHandleTypes::Translate;
@@ -3117,11 +3127,12 @@ EGizmoHandleTypes UViewportWorldInteraction::GetCurrentGizmoType() const
 
 void UViewportWorldInteraction::SetGizmoHandleType( const EGizmoHandleTypes InGizmoHandleType )
 {
+	GizmoType.Reset();
+
 	switch( InGizmoHandleType )
 	{
 		case EGizmoHandleTypes::All:
-			// @todo gizmo: This handle type requires an editor preference to be enabled.  Think about how we actually want to expose the 'universal' gizmo type
-			GLevelEditorModeTools().SetWidgetMode( FWidget::WM_TranslateRotateZ );
+			GizmoType = InGizmoHandleType;
 			break;
 
 		case EGizmoHandleTypes::Translate:
