Fixing non-unity compile errors from missing StaticMeshComponent include

diff --git a/Engine/Source/Editor/VREditor/Teleporter/VREditorTeleporter.cpp b/Engine/Source/Editor/VREditor/Teleporter/VREditorTeleporter.cpp
--- a/Engine/Source/Editor/VREditor/Teleporter/VREditorTeleporter.cpp
+++ b/Engine/Source/Editor/VREditor/Teleporter/VREditorTeleporter.cpp
@@ -7,7 +7,7 @@
 #include "VREditorInteractor.h"
 #include "VREditorMotionControllerInteractor.h"
 #include "VREditorAssetContainer.h"
-
+#include "Components/StaticMeshComponent.h"
 #include "Engine/EngineTypes.h"
 #include "HeadMountedDisplayTypes.h"
 #include "Sound/SoundCue.h"
diff --git a/Engine/Source/Editor/VREditor/UI/VREditorUISystem.cpp b/Engine/Source/Editor/VREditor/UI/VREditorUISystem.cpp
--- a/Engine/Source/Editor/VREditor/UI/VREditorUISystem.cpp
+++ b/Engine/Source/Editor/VREditor/UI/VREditorUISystem.cpp
@@ -51,6 +51,7 @@
 
 // UI
 #include "Components/WidgetComponent.h"
+#include "Components/StaticMeshComponent.h"
 #include "VREditorWidgetComponent.h"
 #include "VREditorStyle.h"
 #include "SUniformGridPanel.h"
diff --git a/Engine/Source/Editor/ViewportInteraction/Gizmo/VIPivotTransformGizmo.cpp b/Engine/Source/Editor/ViewportInteraction/Gizmo/VIPivotTransformGizmo.cpp
--- a/Engine/Source/Editor/ViewportInteraction/Gizmo/VIPivotTransformGizmo.cpp
+++ b/Engine/Source/Editor/ViewportInteraction/Gizmo/VIPivotTransformGizmo.cpp
@@ -16,6 +16,7 @@
 #include "Materials/MaterialInstance.h"
 #include "Materials/MaterialInstanceDynamic.h"
 #include "Components/TextRenderComponent.h"
+#include "Components/StaticMeshComponent.h"
 #include "Math/UnitConversion.h"
 
 namespace VREd //@todo VREditor: Duplicates of TransformGizmo
