Fixes for CIS errors (missing includes, shadow variables)

diff --git a/Engine/Source/Editor/Sequencer/Private/SequencerKeyActor.cpp b/Engine/Source/Editor/Sequencer/Private/SequencerKeyActor.cpp
--- a/Engine/Source/Editor/Sequencer/Private/SequencerKeyActor.cpp
+++ b/Engine/Source/Editor/Sequencer/Private/SequencerKeyActor.cpp
@@ -10,6 +10,8 @@
 #include "EditorModeTools.h"
 #include "EditorModeManager.h"
 #include "ModuleManager.h"
+#include "EditorViewportClient.h"
+#include "UnrealClient.h"
 
 ASequencerKeyActor::ASequencerKeyActor()
 	: Super()
diff --git a/Engine/Source/Editor/Sequencer/Private/SequencerMeshTrail.cpp b/Engine/Source/Editor/Sequencer/Private/SequencerMeshTrail.cpp
--- a/Engine/Source/Editor/Sequencer/Private/SequencerMeshTrail.cpp
+++ b/Engine/Source/Editor/Sequencer/Private/SequencerMeshTrail.cpp
@@ -10,6 +10,7 @@
 #include "Components/StaticMeshComponent.h"
 #include "MovieScene3DTransformSection.h"
 #include "TimerManager.h"
+#include "Editor.h"
 
 ASequencerMeshTrail::ASequencerMeshTrail()
 	: Super()
diff --git a/Engine/Source/Editor/Sequencer/Private/SequencerMeshTrail.h b/Engine/Source/Editor/Sequencer/Private/SequencerMeshTrail.h
--- a/Engine/Source/Editor/Sequencer/Private/SequencerMeshTrail.h
+++ b/Engine/Source/Editor/Sequencer/Private/SequencerMeshTrail.h
@@ -1,6 +1,7 @@
 // Copyright 1998-2017 Epic Games, Inc. All Rights Reserved.
 
 #pragma once
+#include "GameFramework/Actor.h"
 #include "SequencerMeshTrail.generated.h"
 
 /** Stores time and actor reference for each key Actor on the trail */
diff --git a/Engine/Source/Editor/UnrealEd/Private/EditorModeManager.cpp b/Engine/Source/Editor/UnrealEd/Private/EditorModeManager.cpp
--- a/Engine/Source/Editor/UnrealEd/Private/EditorModeManager.cpp
+++ b/Engine/Source/Editor/UnrealEd/Private/EditorModeManager.cpp
@@ -21,6 +21,8 @@
 #include "EditorWorldExtension.h"
 #include "ViewportWorldInteraction.h"
 #include "Editor/EditorEngine.h"
+#include "UnrealEdGlobals.h"
+#include "Editor/UnrealEdEngine.h"
 
 /*------------------------------------------------------------------------------
 	FEditorModeTools.
diff --git a/Engine/Source/Editor/UnrealEd/Private/EditorWorldExtension.cpp b/Engine/Source/Editor/UnrealEd/Private/EditorWorldExtension.cpp
--- a/Engine/Source/Editor/UnrealEd/Private/EditorWorldExtension.cpp
+++ b/Engine/Source/Editor/UnrealEd/Private/EditorWorldExtension.cpp
@@ -5,6 +5,8 @@
 #include "Editor.h"
 #include "ILevelEditor.h"
 #include "LevelEditor.h"
+#include "UnrealEdGlobals.h"
+#include "Editor/UnrealEdEngine.h"
 
 /************************************************************************/
 /* UEditorExtension		                                                */
diff --git a/Engine/Source/Editor/UnrealEd/Public/EditorWorldExtension.h b/Engine/Source/Editor/UnrealEd/Public/EditorWorldExtension.h
--- a/Engine/Source/Editor/UnrealEd/Public/EditorWorldExtension.h
+++ b/Engine/Source/Editor/UnrealEd/Public/EditorWorldExtension.h
@@ -3,6 +3,9 @@
 #pragma once
 
 #include "CoreMinimal.h"
+#include "InputCoreTypes.h"
+#include "Templates/SubclassOf.h"
+#include "Engine/EngineBaseTypes.h"
 #include "UObject/GCObject.h"
 #include "EditorWorldExtension.generated.h"
 
@@ -10,6 +13,9 @@ struct FWorldContext;
 class UEditorWorldExtensionCollection;
 class UEditorWorldExtensionManager;
 class UEditorWorldExtension;
+class FViewport;
+class AActor;
+class FEditorViewportClient;
 
 UCLASS()
 class UNREALED_API UEditorWorldExtension : public UObject
diff --git a/Engine/Source/Editor/VREditor/Public/VREditorMode.h b/Engine/Source/Editor/VREditor/Public/VREditorMode.h
--- a/Engine/Source/Editor/VREditor/Public/VREditorMode.h
+++ b/Engine/Source/Editor/VREditor/Public/VREditorMode.h
@@ -21,6 +21,10 @@ class AActor;
 class FEditorViewportClient;
 class SLevelViewport;
 enum class EAutoKeyMode : uint8;
+class UStaticMesh;
+class UStaticMeshComponent;
+class USoundBase;
+class UMaterialInterface;
 
 // Forward declare the GizmoHandleTypes that is defined in VIBaseTransformGizmo.h
 enum class EGizmoHandleTypes : uint8;
diff --git a/Engine/Source/Editor/VREditor/Teleporter/VREditorTeleporter.cpp b/Engine/Source/Editor/VREditor/Teleporter/VREditorTeleporter.cpp
--- a/Engine/Source/Editor/VREditor/Teleporter/VREditorTeleporter.cpp
+++ b/Engine/Source/Editor/VREditor/Teleporter/VREditorTeleporter.cpp
@@ -82,7 +82,7 @@ void AVREditorTeleporter::Init(UVREditorMode* InMode)
 	{
 		UMaterialInterface* RoomSpaceMaterial = AssetContainer.TeleportMaterial;
 		check(RoomSpaceMaterial != nullptr);
-		TeleportMID = UMaterialInstanceDynamic::Create(RoomSpaceMaterial, GetTransientPackage());
+		TeleportMID = UMaterialInstanceDynamic::Create(RoomSpaceMaterial, this);
 		check(TeleportMID != nullptr);
 	}
 
diff --git a/Engine/Source/Editor/VREditor/Teleporter/VREditorTeleporter.h b/Engine/Source/Editor/VREditor/Teleporter/VREditorTeleporter.h
--- a/Engine/Source/Editor/VREditor/Teleporter/VREditorTeleporter.h
+++ b/Engine/Source/Editor/VREditor/Teleporter/VREditorTeleporter.h
@@ -5,6 +5,8 @@
 #include "CoreMinimal.h"
 #include "UObject/ObjectMacros.h"
 #include "UObject/Object.h"
+#include "GameFramework/Actor.h"
+#include "Engine/EngineBaseTypes.h"
 #include "VREditorTeleporter.generated.h"
 
 class UViewportInteractor;
diff --git a/Engine/Source/Editor/VREditor/UI/VREditorRadialFloatingUI.cpp b/Engine/Source/Editor/VREditor/UI/VREditorRadialFloatingUI.cpp
--- a/Engine/Source/Editor/VREditor/UI/VREditorRadialFloatingUI.cpp
+++ b/Engine/Source/Editor/VREditor/UI/VREditorRadialFloatingUI.cpp
@@ -478,8 +478,8 @@ const void AVREditorRadialFloatingUI::HighlightSlot(const FVector2D& TrackpadPos
 		{
 			NewArrowScaleFactor = 1.0f;
 		}
-		const FVector Scale = FVector(NewArrowScaleFactor, NewArrowScaleFactor, 2.0f * NewArrowScaleFactor);
-		ArrowMeshComponent->SetRelativeScale3D(Scale);
+		const FVector NewArrowScale = FVector(NewArrowScaleFactor, NewArrowScaleFactor, 2.0f * NewArrowScaleFactor);
+		ArrowMeshComponent->SetRelativeScale3D(NewArrowScale);
 	}
 	else
 	{
diff --git a/Engine/Source/Editor/VREditor/UI/VREditorUISystem.h b/Engine/Source/Editor/VREditor/UI/VREditorUISystem.h
--- a/Engine/Source/Editor/VREditor/UI/VREditorUISystem.h
+++ b/Engine/Source/Editor/VREditor/UI/VREditorUISystem.h
@@ -13,6 +13,7 @@
 class AVREditorDockableWindow;
 class AVREditorFloatingUI;
 class FProxyTabmanager;
+struct FTabId;
 class SColorPicker;
 class SBorder;
 class SButton;
diff --git a/Engine/Source/Editor/VREditor/VREditorActions.cpp b/Engine/Source/Editor/VREditor/VREditorActions.cpp
--- a/Engine/Source/Editor/VREditor/VREditorActions.cpp
+++ b/Engine/Source/Editor/VREditor/VREditorActions.cpp
@@ -29,6 +29,8 @@
 #include "Engine/Selection.h"
 #include "VREditorMotionControllerInteractor.h"
 #include "LevelEditorActions.h"
+#include "UObjectIterator.h"
+#include "Factories/Factory.h"
 
 #define LOCTEXT_NAMESPACE "VREditorActions"
 
diff --git a/Engine/Source/Editor/VREditor/VREditorMode.cpp b/Engine/Source/Editor/VREditor/VREditorMode.cpp
--- a/Engine/Source/Editor/VREditor/VREditorMode.cpp
+++ b/Engine/Source/Editor/VREditor/VREditorMode.cpp
@@ -3,6 +3,7 @@
 #include "VREditorMode.h"
 #include "Modules/ModuleManager.h"
 #include "Framework/Application/SlateApplication.h"
+#include "SDockTab.h"
 #include "Engine/EngineTypes.h"
 #include "Components/SceneComponent.h"
 #include "GameFramework/Actor.h"
diff --git a/Engine/Source/Editor/VREditor/VREditorWorldInteraction.cpp b/Engine/Source/Editor/VREditor/VREditorWorldInteraction.cpp
--- a/Engine/Source/Editor/VREditor/VREditorWorldInteraction.cpp
+++ b/Engine/Source/Editor/VREditor/VREditorWorldInteraction.cpp
@@ -466,8 +466,8 @@ void UVREditorWorldInteraction::PlaceDraggedMaterialOrTexture( UViewportInteract
 				if( !bPlaced )
 				{
 					const int32 TargetMaterialSlot = -1;	// All materials
-					bool AppliedMaterial = FComponentEditorUtils::AttemptApplyMaterialToComponent( HitComponent, DroppedObjAsMaterial, TargetMaterialSlot );
-					if (AppliedMaterial)
+					bool bAppliedMaterial = FComponentEditorUtils::AttemptApplyMaterialToComponent( HitComponent, DroppedObjAsMaterial, TargetMaterialSlot );
+					if (bAppliedMaterial)
 					{
 						const UVREditorAssetContainer& AssetContainer = Owner->GetAssetContainer();
 						Owner->PlaySound(AssetContainer.DropFromContentBrowserSound, Interactor->GetTransform().GetLocation());
diff --git a/Engine/Source/Editor/ViewportInteraction/ActorTransformer.cpp b/Engine/Source/Editor/ViewportInteraction/ActorTransformer.cpp
--- a/Engine/Source/Editor/ViewportInteraction/ActorTransformer.cpp
+++ b/Engine/Source/Editor/ViewportInteraction/ActorTransformer.cpp
@@ -6,6 +6,7 @@
 #include "Engine/Selection.h"
 #include "ViewportInteractionAssetContainer.h"
 #include "ViewportInteractor.h"
+#include "Editor.h"
 
 void UActorTransformer::Init( UViewportWorldInteraction* InitViewportWorldInteraction )
 {
diff --git a/Engine/Source/Editor/ViewportInteraction/ActorViewportTransformable.cpp b/Engine/Source/Editor/ViewportInteraction/ActorViewportTransformable.cpp
--- a/Engine/Source/Editor/ViewportInteraction/ActorViewportTransformable.cpp
+++ b/Engine/Source/Editor/ViewportInteraction/ActorViewportTransformable.cpp
@@ -1,6 +1,10 @@
 // Copyright 1998-2016 Epic Games, Inc. All Rights Reserved.
 
 #include "ActorViewportTransformable.h"
+#include "GameFramework/Actor.h"
+#include "Editor.h"
+#include "Editor/EditorEngine.h"
+#include "Components/PrimitiveComponent.h"
 
 
 void FActorViewportTransformable::ApplyTransform( const FTransform& NewTransform, const bool bSweep )
diff --git a/Engine/Source/Editor/ViewportInteraction/MouseCursorInteractor.cpp b/Engine/Source/Editor/ViewportInteraction/MouseCursorInteractor.cpp
--- a/Engine/Source/Editor/ViewportInteraction/MouseCursorInteractor.cpp
+++ b/Engine/Source/Editor/ViewportInteraction/MouseCursorInteractor.cpp
@@ -4,6 +4,7 @@
 #include "SceneView.h"
 #include "EditorViewportClient.h"
 #include "ViewportWorldInteraction.h"
+#include "Editor.h"
 
 
 UMouseCursorInteractor::UMouseCursorInteractor( const FObjectInitializer& Initializer ) 
diff --git a/Engine/Source/Editor/ViewportInteraction/Public/ActorViewportTransformable.h b/Engine/Source/Editor/ViewportInteraction/Public/ActorViewportTransformable.h
--- a/Engine/Source/Editor/ViewportInteraction/Public/ActorViewportTransformable.h
+++ b/Engine/Source/Editor/ViewportInteraction/Public/ActorViewportTransformable.h
@@ -3,6 +3,8 @@
 #pragma once
 
 #include "ViewportTransformable.h"
+#include "UObject/WeakObjectPtrTemplates.h"
+#include "UObject/WeakObjectPtr.h"
 
 /**
  * A transformable actor
diff --git a/Engine/Source/Editor/ViewportInteraction/Public/VIGizmoHandle.h b/Engine/Source/Editor/ViewportInteraction/Public/VIGizmoHandle.h
--- a/Engine/Source/Editor/ViewportInteraction/Public/VIGizmoHandle.h
+++ b/Engine/Source/Editor/ViewportInteraction/Public/VIGizmoHandle.h
@@ -3,6 +3,7 @@
 #pragma once
 
 #include "CoreMinimal.h"
+#include "UnrealWidget.h"
 #include "UObject/ObjectMacros.h"
 #include "Components/SceneComponent.h"
 #include "ViewportInteractionTypes.h"
@@ -11,6 +12,7 @@
 class UMaterialInterface;
 class UStaticMesh;
 enum class EGizmoHandleTypes : uint8;
+class UActorComponent;
 
 USTRUCT()
 struct VIEWPORTINTERACTION_API FGizmoHandle
diff --git a/Engine/Source/Editor/ViewportInteraction/Public/ViewportInteractionUtils.h b/Engine/Source/Editor/ViewportInteraction/Public/ViewportInteractionUtils.h
--- a/Engine/Source/Editor/ViewportInteraction/Public/ViewportInteractionUtils.h
+++ b/Engine/Source/Editor/ViewportInteraction/Public/ViewportInteractionUtils.h
@@ -1,5 +1,9 @@
+// Copyright 1998-2017 Epic Games, Inc. All Rights Reserved.
+
 #pragma once
 
+#include "CoreMinimal.h"
+
 namespace ViewportInteractionUtils
 {
 	/************************************************************************/
diff --git a/Engine/Source/Editor/ViewportInteraction/Public/ViewportTransformable.h b/Engine/Source/Editor/ViewportInteraction/Public/ViewportTransformable.h
--- a/Engine/Source/Editor/ViewportInteraction/Public/ViewportTransformable.h
+++ b/Engine/Source/Editor/ViewportInteraction/Public/ViewportTransformable.h
@@ -2,6 +2,7 @@
 
 #pragma once
 
+#include "CoreMinimal.h"
 
 /**
  * Represents an object that we're actively interacting with, such as a selected actor
diff --git a/Engine/Source/Editor/ViewportInteraction/Public/ViewportTransformer.h b/Engine/Source/Editor/ViewportInteraction/Public/ViewportTransformer.h
--- a/Engine/Source/Editor/ViewportInteraction/Public/ViewportTransformer.h
+++ b/Engine/Source/Editor/ViewportInteraction/Public/ViewportTransformer.h
@@ -2,8 +2,15 @@
 
 #pragma once
 
+#include "CoreMinimal.h"
+#include "UObject/ObjectMacros.h"
+#include "UObject/Object.h"
+
+
 #include "ViewportTransformer.generated.h"
 
+class UViewportWorldInteraction;
+class UViewportInteractor;
 
 UCLASS( abstract )
 class VIEWPORTINTERACTION_API UViewportTransformer : public UObject
@@ -13,7 +20,7 @@ class VIEWPORTINTERACTION_API UViewportTransformer : public UObject
 public:
 
 	UFUNCTION()
-	virtual void Init( class UViewportWorldInteraction* InitViewportWorldInteraction );
+	virtual void Init( UViewportWorldInteraction* InitViewportWorldInteraction );
 
 	UFUNCTION()
 	virtual void Shutdown();
@@ -34,17 +41,17 @@ public:
 
 	/** When starting to drag */
 	UFUNCTION()
-	virtual void OnStartDragging(class UViewportInteractor* Interactor) {};
+	virtual void OnStartDragging(UViewportInteractor* Interactor) {};
 
 	/** When ending drag */
 	UFUNCTION()
-	virtual void OnStopDragging(class UViewportInteractor* Interactor) {};
+	virtual void OnStopDragging(UViewportInteractor* Interactor) {};
 
 protected:
 
 	/** The viewport world interaction object we're registered with */
 	UPROPERTY()
-	class UViewportWorldInteraction* ViewportWorldInteraction;
+	UViewportWorldInteraction* ViewportWorldInteraction;
 
 };
 
diff --git a/Engine/Source/Editor/ViewportInteraction/Public/ViewportWorldInteraction.h b/Engine/Source/Editor/ViewportInteraction/Public/ViewportWorldInteraction.h
--- a/Engine/Source/Editor/ViewportInteraction/Public/ViewportWorldInteraction.h
+++ b/Engine/Source/Editor/ViewportInteraction/Public/ViewportWorldInteraction.h
@@ -19,10 +19,6 @@
 #include "GenericPlatform/ICursor.h"
 #include "ViewportWorldInteraction.generated.h"
 
-
-class IViewportInteractableInterface;
-class UViewportInteractor;
-
 namespace ViewportWorldActionTypes
 {
 	static const FName NoAction( "NoAction" );
@@ -39,7 +35,6 @@ enum class EGizmoHandleTypes : uint8;
 class IViewportInteractableInterface;
 class UViewportInteractor;
 
-
 UCLASS()
 class VIEWPORTINTERACTION_API UViewportWorldInteraction : public UEditorWorldExtension
 {
diff --git a/Engine/Source/Editor/ViewportInteraction/ViewportInteractionUtils.cpp b/Engine/Source/Editor/ViewportInteraction/ViewportInteractionUtils.cpp
--- a/Engine/Source/Editor/ViewportInteraction/ViewportInteractionUtils.cpp
+++ b/Engine/Source/Editor/ViewportInteraction/ViewportInteractionUtils.cpp
@@ -1,3 +1,5 @@
+// Copyright 1998-2017 Epic Games, Inc. All Rights Reserved.
+
 #include "ViewportInteractionUtils.h"
 
 namespace ViewportInteractionUtils
diff --git a/Engine/Source/Editor/ViewportInteraction/ViewportWorldInteraction.cpp b/Engine/Source/Editor/ViewportInteraction/ViewportWorldInteraction.cpp
--- a/Engine/Source/Editor/ViewportInteraction/ViewportWorldInteraction.cpp
+++ b/Engine/Source/Editor/ViewportInteraction/ViewportWorldInteraction.cpp
@@ -2654,7 +2654,7 @@ void UViewportWorldInteraction::RefreshTransformGizmo( const bool bNewObjectsSel
 			TransformGizmoActor->OnNewObjectsSelected();
 		}
 
-		const EGizmoHandleTypes GizmoType = GetCurrentGizmoType();
+		const EGizmoHandleTypes CurrentGizmoType = GetCurrentGizmoType();
 		const ECoordSystem GizmoCoordinateSpace = GetTransformGizmoCoordinateSpace();
 
 		GizmoLocalBounds = GizmoSpaceSelectedObjectsBounds;
@@ -2673,7 +2673,7 @@ void UViewportWorldInteraction::RefreshTransformGizmo( const bool bNewObjectsSel
 		}
 
 		TransformGizmoActor->UpdateGizmo( 
-			GizmoType, 
+			CurrentGizmoType,
 			GizmoCoordinateSpace, 
 			GizmoToWorld, 
 			GizmoSpaceSelectedObjectsBounds, 
diff --git a/Engine/Source/Runtime/Engine/Private/SkeletalRender.cpp b/Engine/Source/Runtime/Engine/Private/SkeletalRender.cpp
--- a/Engine/Source/Runtime/Engine/Private/SkeletalRender.cpp
+++ b/Engine/Source/Runtime/Engine/Private/SkeletalRender.cpp
@@ -81,7 +81,7 @@ void FSkeletalMeshObject::UpdateMinDesiredLODLevel(const FSceneView* View, const
 	int32 NewLODLevel = 0;
 
 	// Look for a lower LOD if the EngineShowFlags is enabled - Thumbnail rendering disables LODs
-	if( 1==View->Family->EngineShowFlags.LOD )
+	if( View && View->Family && 1==View->Family->EngineShowFlags.LOD )
 	{
 		// Iterate from worst to best LOD
 		for(int32 LODLevel = SkeletalMeshResource->LODModels.Num()-1; LODLevel > 0; LODLevel--) 
