Fix for incremental build error - switch to include from forward declare

diff --git a/Engine/Source/Editor/ViewportInteraction/Public/ViewportTransformer.h b/Engine/Source/Editor/ViewportInteraction/Public/ViewportTransformer.h
--- a/Engine/Source/Editor/ViewportInteraction/Public/ViewportTransformer.h
+++ b/Engine/Source/Editor/ViewportInteraction/Public/ViewportTransformer.h
@@ -5,13 +5,10 @@
 #include "CoreMinimal.h"
 #include "UObject/ObjectMacros.h"
 #include "UObject/Object.h"
-
-
+#include "ViewportInteractor.h"
+#include "ViewportWorldInteraction.h"
 #include "ViewportTransformer.generated.h"
 
-class UViewportWorldInteraction;
-class UViewportInteractor;
-
 UCLASS( abstract )
 class VIEWPORTINTERACTION_API UViewportTransformer : public UObject
 {
