Fixing duplicate localization keys

diff --git a/Engine/Source/Editor/VREditor/UI/VREditorUISystem.cpp b/Engine/Source/Editor/VREditor/UI/VREditorUISystem.cpp
--- a/Engine/Source/Editor/VREditor/UI/VREditorUISystem.cpp
+++ b/Engine/Source/Editor/VREditor/UI/VREditorUISystem.cpp
@@ -2075,7 +2075,7 @@ TSharedRef<SWidget> UVREditorUISystem::BuildQuickMenuWidget()
 		EUserInterfaceActionType::ToggleButton
 		);
 	InterfaceMenuBuilder.AddMenuEntry(
-		LOCTEXT("Modes", "Modes"),
+		LOCTEXT("LegacyModes", "Modes"),
 		LOCTEXT("ModesTooltip", "Modes"),
 		FSlateIcon(FEditorStyle::GetStyleSetName(), "LevelEditor.Tabs.Modes"),
 		FUIAction
diff --git a/Engine/Source/Editor/VREditor/UI/VRRadialMenuHandler.cpp b/Engine/Source/Editor/VREditor/UI/VRRadialMenuHandler.cpp
--- a/Engine/Source/Editor/VREditor/UI/VRRadialMenuHandler.cpp
+++ b/Engine/Source/Editor/VREditor/UI/VRRadialMenuHandler.cpp
@@ -485,7 +485,7 @@ void UVRRadialMenuHandler::UIMenuGenerator(FMenuBuilder MenuBuilder, TSharedPtr<
 		EUserInterfaceActionType::ToggleButton
 		);
 	MenuBuilder.AddMenuEntry(
-		LOCTEXT("Modes", "Modes Panel"),
+		LOCTEXT("ModesPanel", "Modes Panel"),
 		FText(),
 		FSlateIcon(FVREditorStyle::GetStyleSetName(), "VREditorStyle.ModesPanel"),
 		FUIAction
@@ -633,7 +633,7 @@ void UVRRadialMenuHandler::ToolsMenuGenerator(FMenuBuilder MenuBuilder, TSharedP
 
 
 	MenuBuilder.AddMenuEntry(
-		LOCTEXT("KeepSimulationChanges", "Save Actors"),
+		LOCTEXT("SaveActors", "Save Actors"),
 		FText(),
 		FSlateIcon(FVREditorStyle::GetStyleSetName(), "VREditorStyle.SaveSimulation"),
 		FUIAction
diff --git a/Engine/Source/Editor/ViewportInteraction/ViewportInteractor.cpp b/Engine/Source/Editor/ViewportInteraction/ViewportInteractor.cpp
--- a/Engine/Source/Editor/ViewportInteraction/ViewportInteractor.cpp
+++ b/Engine/Source/Editor/ViewportInteraction/ViewportInteractor.cpp
@@ -310,7 +310,7 @@ bool UViewportInteractor::HandleInputKey( FEditorViewportClient& ViewportClient,
 												// Capture undo state
 												bSelectionChanged = true;
 												TrackingTransaction.TransCount++;
-												TrackingTransaction.Begin( LOCTEXT( "UndoSelectingActors", "Select Actor" ) );
+												TrackingTransaction.Begin( LOCTEXT( "SelectActor", "Select Actor" ) );
 
 												GEditor->SelectNone( true, true );
 												GEditor->SelectActor( Actor, true, true );
@@ -318,7 +318,7 @@ bool UViewportInteractor::HandleInputKey( FEditorViewportClient& ViewportClient,
 											else if ( SelectionModification == ESelectionModification::Toggle )
 											{
 												TrackingTransaction.TransCount++;
-												TrackingTransaction.Begin( LOCTEXT( "UndoSelectingActors", "Toggle Actor Selection" ) );
+												TrackingTransaction.Begin( LOCTEXT( "ToggleActorSelection", "Toggle Actor Selection" ) );
 
 												GEditor->SelectActor( Actor, !Actor->IsSelected(), true );
 												bSelectionChanged = true;
