VR Editor: - Gizmo animation when aiming at it only check for interactor that has a laser.
- Fix rotation handle animation when selection changed.

diff --git a/Engine/Source/Editor/ViewportInteraction/Gizmo/VIPivotTransformGizmo.cpp b/Engine/Source/Editor/ViewportInteraction/Gizmo/VIPivotTransformGizmo.cpp
--- a/Engine/Source/Editor/ViewportInteraction/Gizmo/VIPivotTransformGizmo.cpp
+++ b/Engine/Source/Editor/ViewportInteraction/Gizmo/VIPivotTransformGizmo.cpp
@@ -115,7 +115,7 @@ void APivotTransformGizmo::UpdateGizmo(const EGizmoHandleTypes InGizmoType, cons
 		for (UViewportInteractor* Interactor : WorldInteraction->GetInteractors())
 		{	
 			// We only want the interactor to affect the size when aiming at the gizmo if they are dragging nothing
-			if (Interactor->GetDraggingMode() == EViewportInteractionDraggingMode::Nothing)
+			if (Interactor->GetDraggingMode() == EViewportInteractionDraggingMode::Nothing && !Interactor->GetIsLaserBlocked())
 			{
 				bAnyInteractorInCorrectState = true;			
 
@@ -446,8 +446,10 @@ void UPivotRotationGizmoHandleGroup::UpdateGizmoHandleGroup( const FTransform& L
 					Rotation.Roll = -90;
 				}
 
+				const float GizmoHandleScale = GizmoScale * AnimationAlpha;
+
 				// Set the final transform
-				GizmoHandleMeshComponent->SetRelativeTransform(FTransform(Rotation, FVector::ZeroVector, FVector(GizmoScale)));
+				GizmoHandleMeshComponent->SetRelativeTransform(FTransform(Rotation, FVector::ZeroVector, FVector(GizmoHandleScale)));
 
 				// Update material
 				UpdateHandleColor(FacingAxisIndex, Handle, DraggingHandle, HoveringOverHandles);
