Copying //UE4/Dev-VREditor to //UE4/Dev-Main (Source: //UE4/Dev-VREditor @ 3352034)

diff --git a/Engine/Source/Editor/VREditor/VREditorWorldInteraction.cpp b/Engine/Source/Editor/VREditor/VREditorWorldInteraction.cpp
--- a/Engine/Source/Editor/VREditor/VREditorWorldInteraction.cpp
+++ b/Engine/Source/Editor/VREditor/VREditorWorldInteraction.cpp
@@ -7,6 +7,7 @@
 #include "Misc/CommandLine.h"
 #include "Modules/ModuleManager.h"
 #include "VREditorMode.h"
+#include "VREditorAssetContainer.h"
 #include "ViewportInteractionTypes.h"
 #include "Materials/MaterialInterface.h"
 #include "Engine/Texture.h"
@@ -14,13 +15,13 @@
 #include "ViewportWorldInteraction.h"
 #include "Engine/BrushBuilder.h"
 #include "Components/PrimitiveComponent.h"
-#include "Kismet/GameplayStatics.h"
 #include "Kismet2/ComponentEditorUtils.h"
 #include "Sound/SoundCue.h"
 #include "Editor.h"
 #include "VREditorUISystem.h"
 #include "VREditorFloatingUI.h"
 #include "VREditorInteractor.h"
+#include "VREditorMotionControllerInteractor.h"
 
 // For actor placement
 #include "ObjectTools.h"
@@ -37,7 +38,6 @@
 namespace VREd
 {
 	static FAutoConsoleVariable SizeOfActorsOverContentBrowserThumbnail( TEXT( "VREd.SizeOfActorsOverContentBrowserThumbnail" ), 6.0f, TEXT( "How large objects should be when rendered 'thumbnail size' over the Content Browser" ) );
-	static FAutoConsoleVariable OverrideSizeOfPlacedActors( TEXT( "VREd.OverrideSizeOfPlacedActors" ), 0.0f, TEXT( "If set to a value greater than zero, sets the size in cm of a placed actor (relative to your world's scale)" ) );
 	static FAutoConsoleVariable HoverHapticFeedbackStrength( TEXT( "VREd.HoverHapticFeedbackStrength" ), 0.1f, TEXT( "Default strength for haptic feedback when hovering" ) );
 	static FAutoConsoleVariable HoverHapticFeedbackTime( TEXT( "VREd.HoverHapticFeedbackTime" ), 0.2f, TEXT( "The minimum time between haptic feedback for hovering" ) );
 	static FAutoConsoleVariable PivotPointTransformGizmo( TEXT( "VREd.PivotPointTransformGizmo" ), 1, TEXT( "If the pivot point transform gizmo is used instead of the bounding box gizmo" ) );
@@ -48,13 +48,9 @@ namespace VREd
 UVREditorWorldInteraction::UVREditorWorldInteraction( const FObjectInitializer& Initializer ) : 
 	Super( Initializer ),
 	Owner( nullptr ),
-	DropMaterialOrMaterialSound( nullptr ),
 	FloatingUIAssetDraggedFrom( nullptr ),
 	PlacingMaterialOrTextureAsset( nullptr )
 {
-	// Load sounds
-	DropMaterialOrMaterialSound = LoadObject<USoundCue>( nullptr, TEXT( "/Engine/VREditor/Sounds/VR_grab_Cue" ) );
-	check( DropMaterialOrMaterialSound != nullptr );
 }
 
 void UVREditorWorldInteraction::Init( UVREditorMode* InOwner, UViewportWorldInteraction* InViewportWorldInteraction )
@@ -77,50 +73,38 @@ void UVREditorWorldInteraction::Shutdown()
 
 	PlacingMaterialOrTextureAsset = nullptr;
 	FloatingUIAssetDraggedFrom = nullptr;
-	DropMaterialOrMaterialSound = nullptr;
 	Owner = nullptr;
 }
 
-//@todo vreditor: This function was overriding ViewportWorldInteraction
-bool UVREditorWorldInteraction::IsInteractableComponent( const UActorComponent* Component ) const
-{
-	bool bResult = false;
-	//bResult = UViewportWorldInteraction::IsInteractableComponent( Component );
-
-	// Don't bother if the base class already found out that this is a custom interactable
-	if ( bResult )
-	{
-		// Don't interact with frozen actors or our UI
-		if ( Component != nullptr )
-		{
-			static const bool bIsVREditorDemo = FParse::Param( FCommandLine::Get(), TEXT( "VREditorDemo" ) );	// @todo vreditor: Remove this when no longer needed (console variable, too!)
-			const bool bIsFrozen = bIsVREditorDemo && Component->GetOwner()->GetActorLabel().StartsWith( TEXT( "Frozen_" ) );
-			bResult = !bIsFrozen && !Owner->GetUISystem().IsWidgetAnEditorUIWidget( Component ); // Don't allow user to move around our UI widgets //@todo VREditor: This wouldn't be necessary if the UI used the interactable interface
-		}
-	}
-	
-	return bResult;
-}
-
 void UVREditorWorldInteraction::StopDragging( UViewportInteractor* Interactor )
 {
 	EViewportInteractionDraggingMode InteractorDraggingMode = Interactor->GetDraggingMode();
 
-	// If we were placing a material, go ahead and do that now
 	if ( InteractorDraggingMode == EViewportInteractionDraggingMode::Material )
 	{
+		// If we were placing a material, go ahead and do that now
 		PlaceDraggedMaterialOrTexture( Interactor );
 	}	
-
-	// If we were placing something, bring the window back
-	if ( FloatingUIAssetDraggedFrom != nullptr && 
-		( InteractorDraggingMode == EViewportInteractionDraggingMode::ActorsAtLaserImpact ||
-		InteractorDraggingMode == EViewportInteractionDraggingMode::Material ) )
+	else if ( FloatingUIAssetDraggedFrom != nullptr && 
+		( InteractorDraggingMode == EViewportInteractionDraggingMode::TransformablesAtLaserImpact ||
+		  InteractorDraggingMode == EViewportInteractionDraggingMode::Material ) )
 	{
-		Owner->GetUISystem().ShowEditorUIPanel( FloatingUIAssetDraggedFrom, Cast<UVREditorInteractor>( Interactor->GetOtherInteractor() ), 
-			true, false, false, false );
+		// If we were placing something, bring the window back
+		const bool bShouldShow = true;
+		const bool bSpawnInFront = false;
+		const bool bDragFromOpen = false;
+		const bool bPlaySound = false;
+		Owner->GetUISystem().ShowEditorUIPanel(FloatingUIAssetDraggedFrom, Cast<UVREditorInteractor>(Interactor->GetOtherInteractor()), bShouldShow, bSpawnInFront, bDragFromOpen, bPlaySound);
 		FloatingUIAssetDraggedFrom = nullptr;
 	}
+	else if (Interactor->GetDraggingMode() == EViewportInteractionDraggingMode::TransformablesFreely)
+	{
+		UVREditorMotionControllerInteractor* MotionController = Cast<UVREditorMotionControllerInteractor>(Interactor);
+		if (Owner->GetUISystem().GetUIInteractor() == MotionController)
+		{
+			MotionController->SetControllerType(EControllerType::UI);
+		}
+	}
 }
 
 void UVREditorWorldInteraction::UpdateNearClipPlaneOnScaleChange(const float NewWorldToMetersScale)
@@ -148,12 +132,14 @@ void UVREditorWorldInteraction::StartDraggingMaterialOrTexture( UViewportInterac
 		FViewportInteractorData& OtherInteractorData = Interactor->GetOtherInteractor()->GetInteractorData();
 		OtherInteractorData.bWasAssistingDrag = false;
 
+		InteractorData.bDraggingWithGrabberSphere = false;
 		InteractorData.bIsFirstDragUpdate = true;
 		InteractorData.bWasAssistingDrag = false;
 		InteractorData.DragRayLength = ( HitLocation - LaserPointerStart ).Size();
-		InteractorData.LastLaserPointerStart = LaserPointerStart;
 		InteractorData.LastDragToLocation = HitLocation;
-		InteractorData.LaserPointerImpactAtDragStart = HitLocation;
+		InteractorData.InteractorRotationAtDragStart = InteractorData.Transform.GetRotation();
+		InteractorData.GrabberSphereLocationAtDragStart = FVector::ZeroVector;
+		InteractorData.ImpactLocationAtDragStart = HitLocation;
 		InteractorData.DragTranslationVelocity = FVector::ZeroVector;
 		InteractorData.DragRayLengthVelocity = 0.0f;
 		InteractorData.bIsDrivingVelocityOfSimulatedTransformables = false;
@@ -162,7 +148,8 @@ void UVREditorWorldInteraction::StartDraggingMaterialOrTexture( UViewportInterac
 	
 		InteractorData.TransformGizmoInteractionType = ETransformGizmoInteractionType::None;
 		InteractorData.GizmoStartTransform = FTransform::Identity;
-		InteractorData.GizmoStartLocalBounds = FBox(ForceInit);
+		InteractorData.GizmoLastTransform = InteractorData.GizmoTargetTransform = InteractorData.GizmoUnsnappedTargetTransform = InteractorData.GizmoInterpolationSnapshotTransform = InteractorData.GizmoStartTransform;
+		InteractorData.GizmoStartLocalBounds = FBox(EForceInit::ForceInit);
 
 		InteractorData.GizmoSpaceFirstDragUpdateOffsetAlongAxis = FVector::ZeroVector;	// Will be determined on first update
 		InteractorData.GizmoSpaceDragDeltaFromStartOffset = FVector::ZeroVector;	// Set every frame while dragging
@@ -177,74 +164,95 @@ void UVREditorWorldInteraction::StartDraggingMaterialOrTexture( UViewportInterac
 
 void UVREditorWorldInteraction::OnAssetDragStartedFromContentBrowser( const TArray<FAssetData>& DraggedAssets, UActorFactory* FactoryToUse )
 {
-	if( ViewportWorldInteraction == nullptr )
+	FloatingUIAssetDraggedFrom = nullptr;
+	if( ViewportWorldInteraction != nullptr )
 	{
-		return;
-	}
+		// Figure out which controller pressed the button and started dragging
+		// @todo vreditor placement: This logic could misfire.  Ideally we would be routed information from the pointer event, so we can determine the hand.
+		UVREditorInteractor* PlacingWithInteractor = nullptr;
 
-	const bool bIsPreview = false;	 // @todo vreditor placement: Consider supporting a "drop preview" actor (so you can cancel placement interactively)
+		TArray<UViewportInteractor*>& Interactors = ViewportWorldInteraction->GetInteractors();
+		for( UViewportInteractor* Interactor : Interactors )
+		{
+			UVREditorInteractor* VRInteractor = Cast<UVREditorInteractor>( Interactor );
+			FViewportActionKeyInput* SelectAndMove_Action = Interactor->GetActionWithName( ViewportWorldActionTypes::SelectAndMove );
+			if( VRInteractor && SelectAndMove_Action != nullptr && SelectAndMove_Action->bIsInputCaptured )
+			{
+				if( VRInteractor->IsClickingOnUI() && !VRInteractor->IsRightClickingOnUI() )
+				{
+					PlacingWithInteractor = VRInteractor;
+					break;
+				}
+			}
+		}
 
-	bool bTransactionStarted = false;
+		if( PlacingWithInteractor != nullptr )
+		{
+			FEditorDelegates::LoadSelectedAssetsIfNeeded.Broadcast();
 
-	// Figure out which controller pressed the button and started dragging
-	// @todo vreditor placement: This logic could misfire.  Ideally we would be routed information from the pointer event, so we can determine the hand.
-	UVREditorInteractor* PlacingWithInteractor = nullptr;
+			TArray< UObject* > DroppedObjects;
+			DroppedObjects.Reserve( DraggedAssets.Num() );
 
-	TArray<UViewportInteractor*>& Interactors = ViewportWorldInteraction->GetInteractors();
-	for ( UViewportInteractor* Interactor : Interactors )
-	{
-		UVREditorInteractor* VRInteractor = Cast<UVREditorInteractor>( Interactor );
-		FViewportActionKeyInput* SelectAndMoveAction = Interactor->GetActionWithName( ViewportWorldActionTypes::SelectAndMove );
-		FViewportActionKeyInput* SelectAndMoveLightlyPressedAction = Interactor->GetActionWithName( ViewportWorldActionTypes::SelectAndMove_LightlyPressed );
+			for( const FAssetData& AssetData : DraggedAssets )
+			{
+				UObject* AssetObj = AssetData.GetAsset();
+				if( AssetObj != nullptr )
+				{
+					// Don't add the same asset more than once
+					DroppedObjects.AddUnique( AssetObj );
+				}
+			}
 
-		if ( VRInteractor && ( ( SelectAndMoveAction && SelectAndMoveAction->bIsInputCaptured ) 
-			|| ( SelectAndMoveLightlyPressedAction && SelectAndMoveLightlyPressedAction->bIsInputCaptured ) ) )
-		{
-			if ( VRInteractor->IsClickingOnUI() && !VRInteractor->IsRightClickingOnUI() ) 
+			if( DroppedObjects.Num() > 0 )
 			{
-				PlacingWithInteractor = VRInteractor;
-				break;
+				// Hide the UI panel that's being used to drag
+				FloatingUIAssetDraggedFrom = PlacingWithInteractor->GetLastHoveredWidgetComponent();
+				Owner->GetUISystem().ShowEditorUIPanel( FloatingUIAssetDraggedFrom, PlacingWithInteractor, false, false, false );
+
+				const bool bShouldInterpolateFromDragLocation = VREd::PlacementInterpolationDuration->GetFloat() > KINDA_SMALL_NUMBER;
+				StartPlacingObjects( DroppedObjects, FactoryToUse, PlacingWithInteractor, bShouldInterpolateFromDragLocation );
+
+				const UVREditorAssetContainer& AssetContainer = Owner->GetAssetContainer();
+				Owner->PlaySound(AssetContainer.DropFromContentBrowserSound, PlacingWithInteractor->GetTransform().GetLocation());
 			}
 		}
 	}
+}
+
+void UVREditorWorldInteraction::StartPlacingObjects( const TArray<UObject*>& ObjectsToPlace, UActorFactory* FactoryToUse, UVREditorInteractor* PlacingWithInteractor, const bool bShouldInterpolateFromDragLocation )
+{
+	if( ViewportWorldInteraction == nullptr || PlacingWithInteractor == nullptr )
+	{
+		return;
+	}
 
+	const bool bIsPreview = false;	 // @todo vreditor placement: Consider supporting a "drop preview" actor (so you can cancel placement interactively)
+
+	bool bTransactionStarted = false;
 
-	// We're always expecting a hand to be hovering at the time we receive this event
-	if( PlacingWithInteractor )
 	{
 		// Cancel UI input
 		{
-			FViewportActionKeyInput* SelectAndMoveAction = PlacingWithInteractor->GetActionWithName( ViewportWorldActionTypes::SelectAndMove );
-			FViewportActionKeyInput* SelectAndMoveLightlyPressedAction = PlacingWithInteractor->GetActionWithName( ViewportWorldActionTypes::SelectAndMove_LightlyPressed );
-
 			PlacingWithInteractor->SetIsClickingOnUI( false );
 			PlacingWithInteractor->SetIsRightClickingOnUI( false );
 
-			if ( SelectAndMoveAction )
+			FViewportActionKeyInput* SelectAndMoveAction = PlacingWithInteractor->GetActionWithName( ViewportWorldActionTypes::SelectAndMove );
+			if( SelectAndMoveAction != nullptr )
 			{
 				SelectAndMoveAction->bIsInputCaptured = false;
 			}
-
-			if ( SelectAndMoveLightlyPressedAction )
-			{
-				SelectAndMoveLightlyPressedAction->bIsInputCaptured = false;
-			}
 		}
 
-		FloatingUIAssetDraggedFrom = PlacingWithInteractor->GetHoveringOverWidgetComponent();
-		// Hide the UI panel that's being used to drag
-		Owner->GetUISystem().ShowEditorUIPanel( FloatingUIAssetDraggedFrom, PlacingWithInteractor, false, false, true, false );
 
 		TArray< UObject* > DroppedObjects;
 		TArray< AActor* > AllNewActors;
 
 		UObject* DraggingSingleMaterialOrTexture = nullptr;
 
-		FVector PlaceAt = PlacingWithInteractor->GetHoverLocation();
+		FVector PlaceAt = PlacingWithInteractor->GetInteractorData().LastHoverLocationOverUI;
 
 		// Only place the object at the laser impact point if we're NOT going to interpolate to the impact
 		// location.  When interpolation is enabled, it looks much better to blend to the new location
-		const bool bShouldInterpolateFromDragLocation = VREd::PlacementInterpolationDuration->GetFloat() > KINDA_SMALL_NUMBER;
 		if( !bShouldInterpolateFromDragLocation )
 		{
 			FVector HitLocation = FVector::ZeroVector;
@@ -254,12 +262,12 @@ void UVREditorWorldInteraction::OnAssetDragStartedFromContentBrowser( const TArr
 			}
 		}
 
-		for( const FAssetData& AssetData : DraggedAssets )
+		for( UObject* AssetObj : ObjectsToPlace )
 		{
 			bool bCanPlace = true;
 
-			UObject* AssetObj = AssetData.GetAsset();
-			if( !ObjectTools::IsAssetValidForPlacing( ViewportWorldInteraction->GetViewportWorld(), AssetData.ObjectPath.ToString() ) )
+			const FString ObjectPath = AssetObj->GetPathName();
+			if( !ObjectTools::IsAssetValidForPlacing( ViewportWorldInteraction->GetWorld(), ObjectPath ) )
 			{
 				bCanPlace = false;
 			}
@@ -278,13 +286,20 @@ void UVREditorWorldInteraction::OnAssetDragStartedFromContentBrowser( const TArr
 			}
 
 			const bool bIsMaterialOrTexture = ( AssetObj->IsA( UMaterialInterface::StaticClass() ) || AssetObj->IsA( UTexture::StaticClass() ) );
-			if( bIsMaterialOrTexture && DraggedAssets.Num() == 1 )
+			if( bIsMaterialOrTexture && ObjectsToPlace.Num() == 1 )
+			{
+				DraggingSingleMaterialOrTexture = AssetObj;
+			}
+			else
 			{
-				DraggingSingleMaterialOrTexture = AssetObj;;
+				// @todo mesheditor: We're dragging actors, so deactivate mesh editor mode for this.  They'll contend over transformables.
+				static FName MeshEditorModeName( TEXT( "MeshEditor" ) );
+				GLevelEditorModeTools().DeactivateMode( MeshEditorModeName );
 			}
 
 			// Check if the asset has an actor factory
-			bool bHasActorFactory = FActorFactoryAssetProxy::GetFactoryForAsset( AssetData ) != NULL;
+			bool bHasActorFactory = FActorFactoryAssetProxy::GetFactoryForAssetObject( AssetObj ) != nullptr;
+
 
 			if( !( AssetObj->IsA( AActor::StaticClass() ) || bHasActorFactory ) &&
 				!AssetObj->IsA( UBrushBuilder::StaticClass() ) )
@@ -315,28 +330,12 @@ void UVREditorWorldInteraction::OnAssetDragStartedFromContentBrowser( const TArr
 				const bool bSelectNewActors = true;
 				const EObjectFlags NewObjectFlags = bIsPreview ? RF_Transient : RF_Transactional;
 
-				TArray<AActor*> NewActors = FLevelEditorViewportClient::TryPlacingActorFromObject( ViewportWorldInteraction->GetViewportWorld()->GetCurrentLevel(), AssetObj, bSelectNewActors, NewObjectFlags, FactoryToUse );
+				TArray<AActor*> NewActors = FLevelEditorViewportClient::TryPlacingActorFromObject( ViewportWorldInteraction->GetWorld()->GetCurrentLevel(), AssetObj, bSelectNewActors, NewObjectFlags, FactoryToUse );
 
+				AllNewActors.Append( NewActors );
 				if( NewActors.Num() > 0 )
 				{
 					DroppedObjects.Add( AssetObj );
-					for( AActor* NewActor : NewActors )
-					{
-						AllNewActors.Add( NewActor );
-
-						if( VREd::OverrideSizeOfPlacedActors->GetFloat() > KINDA_SMALL_NUMBER )
-						{
-							const FBox LocalSpaceBounds = NewActor->CalculateComponentsBoundingBoxInLocalSpace();
-							const float LocalBoundsSize = LocalSpaceBounds.GetSize().GetAbsMax();
-
-							const float DesiredSize = 
-								( bShouldInterpolateFromDragLocation ? 
-									VREd::SizeOfActorsOverContentBrowserThumbnail->GetFloat() :
-									VREd::OverrideSizeOfPlacedActors->GetFloat() )
-								* ViewportWorldInteraction->GetWorldScaleFactor();
-							NewActor->SetActorScale3D( FVector( DesiredSize / LocalBoundsSize ) );
-						}
-					}
 				}
 			}
 		}
@@ -350,7 +349,7 @@ void UVREditorWorldInteraction::OnAssetDragStartedFromContentBrowser( const TArr
 			GEditor->DisableDeltaModification( false );
 		}
 
-		if( AllNewActors.Num() > 0 )	// @todo vreditor: Should we do this for dragged materials too?
+		if( AllNewActors.Num() > 0 )
 		{
 			if( !bIsPreview )
 			{
@@ -361,54 +360,54 @@ void UVREditorWorldInteraction::OnAssetDragStartedFromContentBrowser( const TArr
 
 				FEditorDelegates::OnNewActorsDropped.Broadcast( DroppedObjects, AllNewActors );
 			}
-		}
 
-		if( AllNewActors.Num() > 0 )
-		{
-			const FViewportActionKeyInput Action( ViewportWorldActionTypes::SelectAndMove_LightlyPressed );
+			FBox BoundsOfAllActors;
+			{
+				BoundsOfAllActors.Init();
+				for( AActor* NewActor : AllNewActors )
+				{
+					BoundsOfAllActors += NewActor->CalculateComponentsBoundingBoxInLocalSpace();
+				}
+			}
+			const float BoundsOfAllActorsSize = BoundsOfAllActors.GetSize().GetAbsMax();
+			const float DesiredScale = ( VREd::SizeOfActorsOverContentBrowserThumbnail->GetFloat() / BoundsOfAllActorsSize ) * ViewportWorldInteraction->GetWorldScaleFactor();
+
+			// Start the placed objects off scaled down to match the content browser thumbnail
+			if( bShouldInterpolateFromDragLocation )
+			{
+ 				for( AActor* NewActor : AllNewActors )
+ 				{
+ 					NewActor->SetActorScale3D( FVector( DesiredScale ) );
+ 				}
+			}
+
+			// We changed the initial scale of selected actors, so make sure our transformables and gizmo start transform are up to date.
+			GEditor->NoteSelectionChange();
 
 			// Start dragging the new actor(s)
-			const bool bIsPlacingActors = true;
-			ViewportWorldInteraction->StartDraggingActors( PlacingWithInteractor, Action, ViewportWorldInteraction->GetTransformGizmoActor()->GetRootComponent(), PlaceAt, bIsPlacingActors );
+			UPrimitiveComponent* ClickedTransformGizmoComponent = nullptr;
+			const bool bIsPlacingNewObjects = true;
+			const bool bAllowInterpolationWhenPlacing = bShouldInterpolateFromDragLocation;
+			const bool bStartTransaction = false;
+			const bool bWithGrabberSphere = false;	// Always place using the laser, not the grabber sphere
+			ViewportWorldInteraction->StartDragging( PlacingWithInteractor, ClickedTransformGizmoComponent, PlaceAt, bIsPlacingNewObjects, bAllowInterpolationWhenPlacing, bStartTransaction, bWithGrabberSphere );
 
 			// If we're interpolating, update the target transform of the actors to use our overridden size.  When
 			// we placed them we set their size to be 'thumbnail sized', and we want them to interpolate to
 			// their actual size in the world
 			if( bShouldInterpolateFromDragLocation )
 			{
-				for( TUniquePtr<FViewportTransformable>& TransformablePtr : ViewportWorldInteraction->Transformables )
-				{
-					FViewportTransformable& Transformable = *TransformablePtr;
-
-					float ObjectSize = 1.0f;
-
-					if( VREd::OverrideSizeOfPlacedActors->GetFloat() > KINDA_SMALL_NUMBER )
-					{
-						ObjectSize = 1.0f * ViewportWorldInteraction->GetWorldScaleFactor();
-
-						// Got an actor?
-						AActor* Actor = Cast<AActor>( Transformable.Object.Get() );
-						if( Actor != nullptr )
-						{
-							const FBox LocalSpaceBounds = Actor->CalculateComponentsBoundingBoxInLocalSpace();
-							const float LocalBoundsSize = LocalSpaceBounds.GetSize().GetAbsMax();
-
-							const float DesiredSize = VREd::OverrideSizeOfPlacedActors->GetFloat() * ViewportWorldInteraction->GetWorldScaleFactor();
-							const FVector NewScale( DesiredSize / LocalBoundsSize );
-
-							Transformable.UnsnappedTargetTransform.SetScale3D( NewScale );
-							Transformable.StartTransform.SetScale3D( NewScale );
-							Transformable.LastTransform.SetScale3D( NewScale );
-							Transformable.TargetTransform.SetScale3D( NewScale );
-						}
-					}
-				}
+				const FVector NewScale( 1.0f / DesiredScale );
+				FViewportInteractorData& InteractorData = PlacingWithInteractor->GetInteractorData();
+				InteractorData.GizmoUnsnappedTargetTransform.SetScale3D( NewScale );
+				InteractorData.GizmoLastTransform.SetScale3D( NewScale );
+				InteractorData.GizmoTargetTransform.SetScale3D( NewScale );
 			}
 		}
 
 		if( DraggingSingleMaterialOrTexture != nullptr )
 		{
-			const FViewportActionKeyInput Action( ViewportWorldActionTypes::SelectAndMove_LightlyPressed );
+			const FViewportActionKeyInput Action( ViewportWorldActionTypes::SelectAndMove );
 
 			// Start dragging the material
 			StartDraggingMaterialOrTexture( PlacingWithInteractor,	Action, PlacingWithInteractor->GetHoverLocation(), DraggingSingleMaterialOrTexture );
@@ -424,6 +423,7 @@ void UVREditorWorldInteraction::PlaceDraggedMaterialOrTexture( UViewportInteract
 	{
 		// Check to see if the laser pointer is over something we can drop on
 		UPrimitiveComponent* HitComponent = nullptr;
+		FVector HitLocation = FVector::ZeroVector;
 		{
 			const bool bIgnoreGizmos = true;	// Never place on top of gizmos, just ignore them
 			const bool bEvenIfUIIsInFront = true;	// Don't let the UI block placement
@@ -436,6 +436,7 @@ void UVREditorWorldInteraction::PlaceDraggedMaterialOrTexture( UViewportInteract
 					if( Cast<AVREditorFloatingUI>( HitResult.GetComponent()->GetOwner() ) == nullptr )
 					{
 						HitComponent = HitResult.GetComponent();
+						HitLocation = HitResult.ImpactPoint;
 					}
 				}
 			}
@@ -459,13 +460,19 @@ void UVREditorWorldInteraction::PlaceDraggedMaterialOrTexture( UViewportInteract
 			{
 				// @todo vreditor placement: How do we get the material ID that was dropped on?  Regular editor uses hit proxies.  We may need to augment FHitResult.
 				// @todo vreditor placement: Support optionally dropping on all materials, not only the impacted material
-				const int32 TargetMaterialSlot = -1;	// All materials
-				const bool bPlaced = FComponentEditorUtils::AttemptApplyMaterialToComponent( HitComponent, DroppedObjAsMaterial, TargetMaterialSlot );
-			}
-
-			if( DroppedObjAsMaterial || DroppedObjAsTexture )
-			{
-				UGameplayStatics::PlaySound2D( ViewportWorldInteraction->GetViewportWorld(), DropMaterialOrMaterialSound );
+				bool bPlaced = false;
+				check( Owner );
+				Owner->OnPlaceDraggedMaterial().Broadcast( HitComponent, DroppedObjAsMaterial, bPlaced );
+				if( !bPlaced )
+				{
+					const int32 TargetMaterialSlot = -1;	// All materials
+					bool bAppliedMaterial = FComponentEditorUtils::AttemptApplyMaterialToComponent( HitComponent, DroppedObjAsMaterial, TargetMaterialSlot );
+					if (bAppliedMaterial)
+					{
+						const UVREditorAssetContainer& AssetContainer = Owner->GetAssetContainer();
+						Owner->PlaySound(AssetContainer.DropFromContentBrowserSound, Interactor->GetTransform().GetLocation());
+					}
+				}
 			}
 		}
 	}
@@ -475,7 +482,7 @@ void UVREditorWorldInteraction::PlaceDraggedMaterialOrTexture( UViewportInteract
 
 void UVREditorWorldInteraction::SnapSelectedActorsToGround()
 {
-	TSharedPtr<SLevelViewport> LevelEditorViewport = StaticCastSharedPtr<SLevelViewport>( ViewportWorldInteraction->GetViewportClient()->GetEditorViewportWidget() );
+	TSharedPtr<SLevelViewport> LevelEditorViewport = StaticCastSharedPtr<SLevelViewport>( ViewportWorldInteraction->GetDefaultOptionalViewportClient()->GetEditorViewportWidget() );
 	if ( LevelEditorViewport.IsValid() )
 	{
 		FLevelEditorModule& LevelEditorModule = FModuleManager::GetModuleChecked<FLevelEditorModule>( TEXT( "LevelEditor" ) );
@@ -484,8 +491,8 @@ void UVREditorWorldInteraction::SnapSelectedActorsToGround()
 
 		CommandList->ExecuteAction( Commands.SnapBottomCenterBoundsToFloor.ToSharedRef() );
 
-		// @todo vreditor: This should not be needed after to allow transformables to stop animating the transforms of actors after they come to rest
-		ViewportWorldInteraction->SetupTransformablesForSelectedActors();
+		// Force transformables to refresh
+		GEditor->NoteSelectionChange();
 	}
 }
 
diff --git a/Engine/Source/Editor/VREditor/VREditorWorldInteraction.h b/Engine/Source/Editor/VREditor/VREditorWorldInteraction.h
--- a/Engine/Source/Editor/VREditor/VREditorWorldInteraction.h
+++ b/Engine/Source/Editor/VREditor/VREditorWorldInteraction.h
@@ -30,12 +30,12 @@ public:
 	/** Removes registered event */
 	void Shutdown();
 
-	/** If this component is interactable by an interactor */
-	bool IsInteractableComponent( const UActorComponent* Component ) const;
-
 	/** Snaps the selected objects to the ground */
 	void SnapSelectedActorsToGround();
 
+	/** Try spawn and start placing the specified objects */
+	void StartPlacingObjects( const TArray<UObject*>& ObjectsToPlace, class UActorFactory* FactoryToUse, class UVREditorInteractor* PlacingWithInteractor, const bool bShouldInterpolateFromDragLocation );
+
 protected:
 	/** When an interactor stops dragging */
 	void StopDragging( UViewportInteractor* Interactor );
@@ -62,10 +62,6 @@ protected:
 	UPROPERTY()
 	UViewportWorldInteraction* ViewportWorldInteraction;
 
-	/** Sound for dropping materials and textures */
-	UPROPERTY()
-	class USoundCue* DropMaterialOrMaterialSound;
-
 	//
 	// Dragging object from UI
 	//
diff --git a/Engine/Source/Editor/ViewportInteraction/Gizmo/VIGizmoHandle.cpp b/Engine/Source/Editor/ViewportInteraction/Gizmo/VIGizmoHandle.cpp
--- a/Engine/Source/Editor/ViewportInteraction/Gizmo/VIGizmoHandle.cpp
+++ b/Engine/Source/Editor/ViewportInteraction/Gizmo/VIGizmoHandle.cpp
@@ -9,11 +9,22 @@
 
 UGizmoHandleGroup::UGizmoHandleGroup()
 	: Super(),
-	bShowOnUniversalGizmo( true )
+	GizmoMaterial(nullptr),
+	TranslucentGizmoMaterial(nullptr),
+	Handles(),
+	OwningTransformGizmoActor(nullptr),
+	bShowOnUniversalGizmo(true)
 {
 
 }
 
+UGizmoHandleGroup::~UGizmoHandleGroup()
+{
+	OwningTransformGizmoActor = nullptr;
+	GizmoMaterial = nullptr;
+	TranslucentGizmoMaterial = nullptr;
+}
+
 FTransformGizmoHandlePlacement UGizmoHandleGroup::MakeHandlePlacementForIndex( const int32 HandleIndex ) const
 {
 	FTransformGizmoHandlePlacement HandlePlacement;
@@ -115,7 +126,7 @@ ETransformGizmoInteractionType UGizmoHandleGroup::GetInteractionType() const
 	return ETransformGizmoInteractionType::Translate;
 }
 
-void UGizmoHandleGroup::UpdateGizmoHandleGroup(const FTransform& LocalToWorld, const FBox& LocalBounds, const FVector ViewLocation, bool bAllHandlesVisible, class UActorComponent* DraggingHandle, 
+void UGizmoHandleGroup::UpdateGizmoHandleGroup(const FTransform& LocalToWorld, const FBox& LocalBounds, const FVector ViewLocation, const bool bAllHandlesVisible, class UActorComponent* DraggingHandle, 
 	const TArray< UActorComponent* >& HoveringOverHandles, float AnimationAlpha, float GizmoScale, const float GizmoHoverScale, const float GizmoHoverAnimationDuration, bool& bOutIsHoveringOrDraggingThisHandleGroup )
 {
 	UpdateHoverAnimation(DraggingHandle, HoveringOverHandles, GizmoHoverAnimationDuration, bOutIsHoveringOrDraggingThisHandleGroup);
@@ -163,6 +174,11 @@ bool UGizmoHandleGroup::GetShowOnUniversalGizmo() const
 	return bShowOnUniversalGizmo;
 }
 
+void UGizmoHandleGroup::SetOwningTransformGizmo(class ABaseTransformGizmo* TransformGizmo)
+{
+	OwningTransformGizmoActor = TransformGizmo;
+}
+
 void UGizmoHandleGroup::UpdateHandleColor( const int32 AxisIndex, FGizmoHandle& Handle, class UActorComponent* DraggingHandle, const TArray< UActorComponent* >& HoveringOverHandles )
 {
 	UStaticMeshComponent* HandleMesh = Handle.HandleMesh;
@@ -181,42 +197,64 @@ void UGizmoHandleGroup::UpdateHandleColor( const int32 AxisIndex, FGizmoHandle&
 	UMaterialInstanceDynamic* MID1 = CastChecked<UMaterialInstanceDynamic>( HandleMesh->GetMaterial( 1 ) );
 
 	ABaseTransformGizmo* GizmoActor = CastChecked<ABaseTransformGizmo>( GetOwner() );
-	if ( GizmoActor )
+	if (GizmoActor)
 	{
 		UViewportWorldInteraction* WorldInteraction = GizmoActor->GetOwnerWorldInteraction();
-		if ( WorldInteraction )
+		if (WorldInteraction)
 		{
-			FLinearColor HandleColor = WorldInteraction->GetColor( UViewportWorldInteraction::EColors::DefaultColor );
-			if ( HandleMesh == DraggingHandle )
+			FLinearColor HandleColor = WorldInteraction->GetColor(UViewportWorldInteraction::EColors::DefaultColor);
+			if (HandleMesh == DraggingHandle)
 			{
-				HandleColor = WorldInteraction->GetColor( UViewportWorldInteraction::EColors::Dragging );
+				HandleColor = WorldInteraction->GetColor(UViewportWorldInteraction::EColors::GizmoDragging);
 			}
-			else if ( AxisIndex != INDEX_NONE )
+			else if (AxisIndex != INDEX_NONE)
 			{
-				switch ( AxisIndex )
+				switch (AxisIndex)
 				{
 				case 0:
-					HandleColor = WorldInteraction->GetColor( UViewportWorldInteraction::EColors::Forward );
+					HandleColor = WorldInteraction->GetColor(UViewportWorldInteraction::EColors::Forward);
 					break;
 
 				case 1:
-					HandleColor = WorldInteraction->GetColor( UViewportWorldInteraction::EColors::Right );
+					HandleColor = WorldInteraction->GetColor(UViewportWorldInteraction::EColors::Right);
 					break;
 
 				case 2:
-					HandleColor = WorldInteraction->GetColor( UViewportWorldInteraction::EColors::Up );
+					HandleColor = WorldInteraction->GetColor(UViewportWorldInteraction::EColors::Up);
 					break;
 				}
 
-				if ( HoveringOverHandles.Contains( HandleMesh ) )
+				if (HoveringOverHandles.Contains(HandleMesh))
 				{
-					HandleColor = FLinearColor::LerpUsingHSV( HandleColor, WorldInteraction->GetColor( UViewportWorldInteraction::EColors::Hover ), Handle.HoverAlpha );
+					HandleColor = FLinearColor::LerpUsingHSV(HandleColor, WorldInteraction->GetColor(UViewportWorldInteraction::EColors::GizmoHover), Handle.HoverAlpha);
 				}
 			}
 
-			static FName StaticHandleColorParameter( "Color" );
-			MID0->SetVectorParameterValue( StaticHandleColorParameter, HandleColor );
-			MID1->SetVectorParameterValue( StaticHandleColorParameter, HandleColor );
+			static FName StaticHandleColorParameter("Color");
+			MID0->SetVectorParameterValue(StaticHandleColorParameter, HandleColor);
+			MID1->SetVectorParameterValue(StaticHandleColorParameter, HandleColor);
+		}
+	}
+}
+
+void UGizmoHandleGroup::UpdateVisibilityAndCollision(const EGizmoHandleTypes GizmoType, const ECoordSystem GizmoCoordinateSpace, const bool bAllHandlesVisible, const bool bAllowRotationAndScaleHandles, UActorComponent* DraggingHandle)
+{
+	const bool bIsTypeSupported = 
+		( (GizmoType == EGizmoHandleTypes::All && GetShowOnUniversalGizmo()) || GetHandleType() == GizmoType ) &&
+		( bAllowRotationAndScaleHandles || ( GetHandleType() != EGizmoHandleTypes::Rotate && GetHandleType() != EGizmoHandleTypes::Scale ) );
+
+	const bool bSupportsCurrentCoordinateSpace = SupportsWorldCoordinateSpace() || GizmoCoordinateSpace != COORD_World;
+
+	for (FGizmoHandle& Handle : GetHandles())
+	{
+		if (Handle.HandleMesh != nullptr)
+		{
+			const bool bShowIt = (bIsTypeSupported && bSupportsCurrentCoordinateSpace && bAllHandlesVisible) || (DraggingHandle != nullptr && DraggingHandle == Handle.HandleMesh);
+
+			Handle.HandleMesh->SetVisibility(bShowIt);
+
+			// Never allow ray queries to impact hidden handles
+			Handle.HandleMesh->SetCollisionEnabled(bShowIt ? ECollisionEnabled::QueryOnly : ECollisionEnabled::NoCollision);
 		}
 	}
 }
@@ -340,47 +378,48 @@ void UAxisGizmoHandleGroup::UpdateHandlesRelativeTransformOnAxis( const FTransfo
 			int32 CenterHandleCount, FacingAxisIndex, CenterAxisIndex;
 			HandlePlacement.GetCenterHandleCountAndFacingAxisIndex( /* Out */ CenterHandleCount, /* Out */ FacingAxisIndex, /* Out */ CenterAxisIndex);
 			
-			FVector GizmoSpaceFacingAxisVector = GetAxisVector( FacingAxisIndex, HandlePlacement.Axes[FacingAxisIndex] );
-			
-			// Check on which side we are relative to the gizmo
-			const FVector GizmoSpaceViewLocation = GetOwner()->GetTransform().InverseTransformPosition( ViewLocation );
-			if ( GizmoSpaceViewLocation[ FacingAxisIndex ] < 0 )
+			if (DraggingHandle == nullptr)
 			{
-				GizmoSpaceFacingAxisVector[ FacingAxisIndex ] *= -1.0f;
-			}
+				FVector GizmoSpaceFacingAxisVector = GetAxisVector( FacingAxisIndex, HandlePlacement.Axes[FacingAxisIndex] );
+			
+				// Check on which side we are relative to the gizmo
+				const FVector GizmoSpaceViewLocation = GetOwner()->GetTransform().InverseTransformPosition( ViewLocation );
+				if ( GizmoSpaceViewLocation[ FacingAxisIndex ] < 0 )
+				{
+					GizmoSpaceFacingAxisVector[ FacingAxisIndex ] *= -1.0f;
+				}
 
-			const FTransform GizmoOriginToFacingAxisRotation( GizmoSpaceFacingAxisVector.ToOrientationQuat() );
+				const FTransform GizmoOriginToFacingAxisRotation( GizmoSpaceFacingAxisVector.ToOrientationQuat() );
 	
-			FTransform HandleToGizmoOrigin = HandleToCenter * GizmoOriginToFacingAxisRotation;
+				FTransform HandleToGizmoOrigin = HandleToCenter * GizmoOriginToFacingAxisRotation;
 				
-			// Check on the axis if the offset is on the other side of the object compared to the view. Switch the offset to the side of the view if it does
-			FVector GizmoSpaceFacingAxisOffset = HandleToGizmoOrigin.GetLocation();
-			for ( int32 AxisIndex = 0; AxisIndex < 3; ++AxisIndex )
-			{		
-				if ( AxisIndex != FacingAxisIndex &&
-					( ( GizmoSpaceFacingAxisOffset[ AxisIndex ] > 0 && GizmoSpaceViewLocation[ AxisIndex ] < 0 ) ||
-					  ( GizmoSpaceFacingAxisOffset[ AxisIndex ] < 0 && GizmoSpaceViewLocation[ AxisIndex ] > 0 ) ) )
-				{
-					GizmoSpaceFacingAxisOffset[ AxisIndex ] *= -1.0f;
-				}	
-			}
+				// Check on the axis if the offset is on the other side of the object compared to the view. Switch the offset to the side of the view if it does
+				FVector GizmoSpaceFacingAxisOffset = HandleToGizmoOrigin.GetLocation();
+				for ( int32 AxisIndex = 0; AxisIndex < 3; ++AxisIndex )
+				{		
+					if ( AxisIndex != FacingAxisIndex &&
+						( ( GizmoSpaceFacingAxisOffset[ AxisIndex ] > 0 && GizmoSpaceViewLocation[ AxisIndex ] < 0 ) ||
+						  ( GizmoSpaceFacingAxisOffset[ AxisIndex ] < 0 && GizmoSpaceViewLocation[ AxisIndex ] > 0 ) ) )
+					{
+						GizmoSpaceFacingAxisOffset[ AxisIndex ] *= -1.0f;
+					}	
+				}
 
-			GizmoSpaceFacingAxisOffset *= AnimationAlpha;
+				GizmoSpaceFacingAxisOffset *= AnimationAlpha;
 
-			HandleToGizmoOrigin.SetLocation( GizmoSpaceFacingAxisOffset * GizmoScale );
+				HandleToGizmoOrigin.SetLocation( GizmoSpaceFacingAxisOffset * GizmoScale );
 
-			// Set the final transform
-			GizmoHandleMeshComponent->SetRelativeTransform( HandleToGizmoOrigin );
+				// Set the final transform
+				GizmoHandleMeshComponent->SetRelativeTransform( HandleToGizmoOrigin );
 			
-			float GizmoHandleScale = GizmoScale;
-			if ( GizmoScale != GizmoHoverScale )
-			{
+				float GizmoHandleScale = GizmoScale;
+
 				// Make the handle bigger while hovered (but don't affect the offset -- we want it to scale about it's origin)
 				GizmoHandleScale *= FMath::Lerp( 1.0f, GizmoHoverScale, Handle.HoverAlpha );
-			}
 			
-			GizmoHandleScale *= AnimationAlpha;
-			GizmoHandleMeshComponent->SetRelativeScale3D( FVector( GizmoHandleScale ) );
+				GizmoHandleScale *= AnimationAlpha;
+				GizmoHandleMeshComponent->SetRelativeScale3D( FVector( GizmoHandleScale ) );
+			}
 
 			// Update material
 			UpdateHandleColor( FacingAxisIndex, Handle, DraggingHandle, HoveringOverHandles );
diff --git a/Engine/Source/Editor/ViewportInteraction/Gizmo/VIPivotTransformGizmo.cpp b/Engine/Source/Editor/ViewportInteraction/Gizmo/VIPivotTransformGizmo.cpp
--- a/Engine/Source/Editor/ViewportInteraction/Gizmo/VIPivotTransformGizmo.cpp
+++ b/Engine/Source/Editor/ViewportInteraction/Gizmo/VIPivotTransformGizmo.cpp
@@ -7,52 +7,74 @@
 #include "GameFramework/WorldSettings.h"
 #include "VIStretchGizmoHandle.h"
 #include "VIUniformScaleGizmoHandle.h"
+#include "ViewportWorldInteraction.h"
+#include "ViewportInteractor.h"
+#include "ViewportInteractionAssetContainer.h"
+#include "Materials/Material.h"
+#include "Engine/Font.h"
+#include "Engine/CollisionProfile.h"
+#include "Materials/MaterialInstance.h"
+#include "Materials/MaterialInstanceDynamic.h"
+#include "Components/TextRenderComponent.h"
+#include "Components/StaticMeshComponent.h"
+#include "Math/UnitConversion.h"
 
 namespace VREd //@todo VREditor: Duplicates of TransformGizmo
 {
 	// @todo vreditor tweak: Tweak out console variables
-	static FAutoConsoleVariable PivotGizmoDistanceScaleFactor( TEXT( "VREd.PivotGizmoDistanceScaleFactor" ), 0.003f, TEXT( "How much the gizmo handles should increase in size with distance from the camera, to make it easier to select"));
-	static FAutoConsoleVariable PivotGizmoTranslationPivotOffsetX( TEXT("VREd.PivotGizmoTranslationPivotOffsetX" ), 20.0f, TEXT( "How much the translation cylinder is offsetted from the pivot" ) );
-	static FAutoConsoleVariable PivotGizmoScalePivotOffsetX( TEXT( "VREd.PivotGizmoScalePivotOffsetX" ), 120.0f, TEXT( "How much the non-uniform scale is offsetted from the pivot" ) );
-	static FAutoConsoleVariable PivotGizmoPlaneTranslationPivotOffsetYZ(TEXT("VREd.PivotGizmoPlaneTranslationPivotOffsetYZ" ), 40.0f, TEXT( "How much the plane translation is offsetted from the pivot" ) );
-	static FAutoConsoleVariable PivotGizmoTranslationScaleMultiply( TEXT( "VREd.PivotGizmoTranslationScaleMultiply" ), 3.0f, TEXT( "Multiplies translation handles scale" ) );
-	static FAutoConsoleVariable PivotGizmoTranslationHoverScaleMultiply( TEXT( "VREd.PivotGizmoTranslationHoverScaleMultiply" ), 1.2f, TEXT( "Multiplies translation handles hover scale" ) );
+	static FAutoConsoleVariable PivotGizmoMinDistanceForScaling( TEXT( "VI.PivotGizmoMinDistanceForScaling" ), 0.0f, TEXT( "How far away the camera needs to be from an object before we'll start scaling it based on distance"));
+	static FAutoConsoleVariable PivotGizmoDistanceScaleFactor( TEXT( "VI.PivotGizmoDistanceScaleFactor" ), 0.0035f, TEXT( "How much the gizmo handles should increase in size with distance from the camera, to make it easier to select" ) );
+	static FAutoConsoleVariable PivotGizmoTranslationPivotOffsetX( TEXT("VI.PivotGizmoTranslationPivotOffsetX" ), 30.0f, TEXT( "How much the translation cylinder is offsetted from the pivot" ) );
+	static FAutoConsoleVariable PivotGizmoScalePivotOffsetX( TEXT( "VI.PivotGizmoScalePivotOffsetX" ), 120.0f, TEXT( "How much the non-uniform scale is offsetted from the pivot" ) );
+	static FAutoConsoleVariable PivotGizmoPlaneTranslationPivotOffsetYZ(TEXT("VI.PivotGizmoPlaneTranslationPivotOffsetYZ" ), 40.0f, TEXT( "How much the plane translation is offsetted from the pivot" ) );
+	static FAutoConsoleVariable PivotGizmoTranslationScaleMultiply( TEXT( "VI.PivotGizmoTranslationScaleMultiply" ), 2.0f, TEXT( "Multiplies translation handles scale" ) );
+	static FAutoConsoleVariable PivotGizmoTranslationHoverScaleMultiply( TEXT( "VI.PivotGizmoTranslationHoverScaleMultiply" ), 0.75f, TEXT( "Multiplies translation handles hover scale" ) );
+	static FAutoConsoleVariable PivotGizmoAimAtShrinkSize( TEXT( "VI.PivotGizmoAimAtShrinkSize" ), 0.3f, TEXT( "The minimum size when not aiming at the gizmo (0 to 1)" ) );
+	static FAutoConsoleVariable PivotGizmoAimAtAnimationSpeed( TEXT( "VI.PivotGizmoAimAtAnimationSpeed" ), 0.15f, TEXT( "The speed to animate to the gizmo full size when aiming at it" ) );
 }
 
 APivotTransformGizmo::APivotTransformGizmo() :
-	Super()
+	Super(),
+	AimingAtGizmoScaleAlpha(0.0f),
+	LastDraggingHandle(nullptr)
 {
 	UniformScaleGizmoHandleGroup = CreateDefaultSubobject<UUniformScaleGizmoHandleGroup>( TEXT( "UniformScaleHandles" ), true );
+	UniformScaleGizmoHandleGroup->SetOwningTransformGizmo(this);
 	UniformScaleGizmoHandleGroup->SetTranslucentGizmoMaterial( TranslucentGizmoMaterial );
 	UniformScaleGizmoHandleGroup->SetGizmoMaterial( GizmoMaterial );
 	UniformScaleGizmoHandleGroup->SetupAttachment( SceneComponent );
 	AllHandleGroups.Add( UniformScaleGizmoHandleGroup );
 
 	TranslationGizmoHandleGroup = CreateDefaultSubobject<UPivotTranslationGizmoHandleGroup>(TEXT("TranslationHandles"), true);
+	TranslationGizmoHandleGroup->SetOwningTransformGizmo(this);
 	TranslationGizmoHandleGroup->SetTranslucentGizmoMaterial(TranslucentGizmoMaterial);
 	TranslationGizmoHandleGroup->SetGizmoMaterial(GizmoMaterial);
 	TranslationGizmoHandleGroup->SetupAttachment( SceneComponent );
 	AllHandleGroups.Add(TranslationGizmoHandleGroup);
 
 	ScaleGizmoHandleGroup = CreateDefaultSubobject<UPivotScaleGizmoHandleGroup>( TEXT( "ScaleHandles" ), true );
+	ScaleGizmoHandleGroup->SetOwningTransformGizmo(this);
 	ScaleGizmoHandleGroup->SetTranslucentGizmoMaterial( TranslucentGizmoMaterial );
 	ScaleGizmoHandleGroup->SetGizmoMaterial( GizmoMaterial );
 	ScaleGizmoHandleGroup->SetupAttachment( SceneComponent );
 	AllHandleGroups.Add( ScaleGizmoHandleGroup );
 
 	PlaneTranslationGizmoHandleGroup = CreateDefaultSubobject<UPivotPlaneTranslationGizmoHandleGroup>( TEXT( "PlaneTranslationHandles" ), true );
+	PlaneTranslationGizmoHandleGroup->SetOwningTransformGizmo(this);
 	PlaneTranslationGizmoHandleGroup->SetTranslucentGizmoMaterial( TranslucentGizmoMaterial );
 	PlaneTranslationGizmoHandleGroup->SetGizmoMaterial( GizmoMaterial );
 	PlaneTranslationGizmoHandleGroup->SetupAttachment( SceneComponent );
 	AllHandleGroups.Add( PlaneTranslationGizmoHandleGroup );
 
 	RotationGizmoHandleGroup = CreateDefaultSubobject<UPivotRotationGizmoHandleGroup>( TEXT( "RotationHandles" ), true );
+	RotationGizmoHandleGroup->SetOwningTransformGizmo(this);
 	RotationGizmoHandleGroup->SetTranslucentGizmoMaterial( TranslucentGizmoMaterial );
 	RotationGizmoHandleGroup->SetGizmoMaterial( GizmoMaterial );
 	RotationGizmoHandleGroup->SetupAttachment( SceneComponent );
 	AllHandleGroups.Add( RotationGizmoHandleGroup );
 
 	StretchGizmoHandleGroup = CreateDefaultSubobject<UStretchGizmoHandleGroup>( TEXT( "StretchHandles" ), true );
+	StretchGizmoHandleGroup->SetOwningTransformGizmo(this);
 	StretchGizmoHandleGroup->SetTranslucentGizmoMaterial( TranslucentGizmoMaterial );
 	StretchGizmoHandleGroup->SetGizmoMaterial( GizmoMaterial );
 	StretchGizmoHandleGroup->SetShowOnUniversalGizmo( false );
@@ -64,38 +86,91 @@ APivotTransformGizmo::APivotTransformGizmo() :
 	OnNewObjectsSelected();
 }
 
-void APivotTransformGizmo::UpdateGizmo( const EGizmoHandleTypes GizmoType, const ECoordSystem GizmoCoordinateSpace, const FTransform& LocalToWorld, const FBox& LocalBounds, const FVector ViewLocation, bool bAllHandlesVisible, 
-	UActorComponent* DraggingHandle, const TArray< UActorComponent* >& HoveringOverHandles, const float GizmoHoverScale, const float GizmoHoverAnimationDuration )
+void APivotTransformGizmo::UpdateGizmo(const EGizmoHandleTypes InGizmoType, const ECoordSystem InGizmoCoordinateSpace, const FTransform& InLocalToWorld, const FBox& InLocalBounds, const FVector& InViewLocation, const float InScaleMultiplier, bool bInAllHandlesVisible, const bool bInAllowRotationAndScaleHandles, class UActorComponent* DraggingHandle, const TArray<UActorComponent*>& InHoveringOverHandles, const float InGizmoHoverScale, const float InGizmoHoverAnimationDuration)
 {
-	const float WorldScaleFactor = GetWorld()->GetWorldSettings()->WorldToMeters / 100.0f;
+	GizmoType = InGizmoType;
 
 	// Position the gizmo at the location of the first selected actor
-	const bool bSweep = false;
-	this->SetActorTransform( LocalToWorld, bSweep );
+	SetActorTransform(InLocalToWorld);
 
 	// Increase scale with distance, to make gizmo handles easier to click on
-	// @todo vreditor: Should probably be a curve, not linear
-	// @todo vreditor: Should take FOV into account (especially in non-stereo/HMD mode)
-	const float WorldSpaceDistanceToToPivot = FMath::Sqrt( FVector::DistSquared( GetActorLocation(), ViewLocation ) );
-	const float GizmoScale( ( ( WorldSpaceDistanceToToPivot / WorldScaleFactor ) * VREd::PivotGizmoDistanceScaleFactor->GetFloat() ) * WorldScaleFactor );
+	const float WorldSpaceDistanceToToPivot = FMath::Max(VREd::PivotGizmoMinDistanceForScaling->GetFloat(), FMath::Sqrt(FVector::DistSquared( GetActorLocation(), InViewLocation )));
+	const float WorldScaleFactor = WorldInteraction->GetWorldScaleFactor();
+	float GizmoScale = (InScaleMultiplier * ((WorldSpaceDistanceToToPivot / WorldScaleFactor) * VREd::PivotGizmoDistanceScaleFactor->GetFloat())) * WorldScaleFactor;
+	const bool bIsWorldSpaceGizmo = (WorldInteraction->GetTransformGizmoCoordinateSpace() == COORD_World);
+
+	if (LastDraggingHandle != nullptr && DraggingHandle == nullptr)
+	{
+		AimingAtGizmoScaleAlpha = 0.0f;
+	}
+
+	// Only scale the gizmo down when not aiming at it for VR implementations
+	if (WorldInteraction->IsInVR())
+	{
+		bool bIsAimingTowards = false;
+		bool bAnyInteractorInCorrectState = false;
+		const float GizmoRadius = (GizmoScale * 350) * 0.5f;
+
+		// Check if any interactor has a laser close to the gizmo
+		for (UViewportInteractor* Interactor : WorldInteraction->GetInteractors())
+		{	
+			// We only want the interactor to affect the size when aiming at the gizmo if they are dragging nothing
+			if (Interactor->GetDraggingMode() == EViewportInteractionDraggingMode::Nothing && !Interactor->GetIsLaserBlocked())
+			{
+				bAnyInteractorInCorrectState = true;			
+
+				FVector LaserPointerStart, LaserPointerEnd;
+				if (Interactor->GetLaserPointer(/*Out*/ LaserPointerStart, /*Out*/ LaserPointerEnd))
+				{
+					const FVector ClosestPointOnLaser = FMath::ClosestPointOnLine(LaserPointerStart, LaserPointerEnd, InLocalToWorld.GetLocation());
+					const float ClosestPointDistance = (ClosestPointOnLaser - InLocalToWorld.GetLocation()).Size();
+					if (ClosestPointDistance <= GizmoRadius)
+					{
+						bIsAimingTowards = true;
+						break;
+					}
+				}
+			}
+		}
+
+		if (bAnyInteractorInCorrectState)
+		{
+			const float DeltaTime = WorldInteraction->GetCurrentDeltaTime();
+			if (bIsAimingTowards)
+			{
+				AimingAtGizmoScaleAlpha += DeltaTime / VREd::PivotGizmoAimAtAnimationSpeed->GetFloat();
+			}
+			else
+			{
+				AimingAtGizmoScaleAlpha -= DeltaTime / VREd::PivotGizmoAimAtAnimationSpeed->GetFloat();
+			}
+			AimingAtGizmoScaleAlpha = FMath::Clamp(AimingAtGizmoScaleAlpha, VREd::PivotGizmoAimAtShrinkSize->GetFloat(), 1.0f);
+
+			GizmoScale *= AimingAtGizmoScaleAlpha;
+		}
+	}
 
 	// Update animation
 	float AnimationAlpha = GetAnimationAlpha();
 
 	// Update all the handles
-	bool bIsHoveringOrDraggingScaleGizmo = false;
-	for ( UGizmoHandleGroup* HandleGroup : AllHandleGroups )
+	for (UGizmoHandleGroup* HandleGroup : AllHandleGroups)
 	{
-		if ( HandleGroup != nullptr && 
-		   ( DraggingHandle == nullptr || HandleGroup == StretchGizmoHandleGroup ) )
+		if (HandleGroup != nullptr)
 		{
+			//const FTransform HandleGroupTransform = bIsWorldSpaceGizmo == true && HandleGroup != RotationGizmoHandleGroup == true ? TransformGizmoToWorld : InLocalToWorld;
 			bool bIsHoveringOrDraggingThisHandleGroup = false;
-			HandleGroup->UpdateGizmoHandleGroup( LocalToWorld, LocalBounds, ViewLocation, bAllHandlesVisible, DraggingHandle,
-				HoveringOverHandles, AnimationAlpha, GizmoScale, GizmoHoverScale, GizmoHoverAnimationDuration, /* Out */ bIsHoveringOrDraggingThisHandleGroup );
+			HandleGroup->UpdateGizmoHandleGroup(InLocalToWorld, InLocalBounds, InViewLocation, bInAllHandlesVisible, DraggingHandle,
+				InHoveringOverHandles, AnimationAlpha, GizmoScale, InGizmoHoverScale, InGizmoHoverAnimationDuration, /* Out */ bIsHoveringOrDraggingThisHandleGroup);
+			
+			if (HandleGroup != RotationGizmoHandleGroup)
+			{
+				HandleGroup->UpdateVisibilityAndCollision(InGizmoType, InGizmoCoordinateSpace, bInAllHandlesVisible, bInAllowRotationAndScaleHandles, DraggingHandle);
+			}
 		}
 	}
 
-	UpdateHandleVisibility( GizmoType, GizmoCoordinateSpace, bAllHandlesVisible, DraggingHandle );
+	LastDraggingHandle = DraggingHandle;
 }
 
 /************************************************************************/
@@ -104,18 +179,12 @@ void APivotTransformGizmo::UpdateGizmo( const EGizmoHandleTypes GizmoType, const
 UPivotTranslationGizmoHandleGroup::UPivotTranslationGizmoHandleGroup() :
 	Super()
 {
-	UStaticMesh* TranslationHandleMesh = nullptr;
-	{
-		static ConstructorHelpers::FObjectFinder<UStaticMesh> ObjectFinder( TEXT( "/Engine/VREditor/TransformGizmo/TranslateHandleLong" ) );
-		TranslationHandleMesh = ObjectFinder.Object;
-		check( TranslationHandleMesh != nullptr );
-	}
-
-	CreateHandles( TranslationHandleMesh, FString( "PivotTranslationHandle" ) );
+	const UViewportInteractionAssetContainer& AssetContainer = UViewportWorldInteraction::LoadAssetContainer();
+	CreateHandles( AssetContainer.TranslationHandleMesh, FString( "PivotTranslationHandle" ) );
 }
 
 
-void UPivotTranslationGizmoHandleGroup::UpdateGizmoHandleGroup( const FTransform& LocalToWorld, const FBox& LocalBounds, const FVector ViewLocation, bool bAllHandlesVisible, class UActorComponent* DraggingHandle, 
+void UPivotTranslationGizmoHandleGroup::UpdateGizmoHandleGroup( const FTransform& LocalToWorld, const FBox& LocalBounds, const FVector ViewLocation, const bool bAllHandlesVisible, class UActorComponent* DraggingHandle, 
 	const TArray< UActorComponent* >& HoveringOverHandles, float AnimationAlpha, float GizmoScale, const float GizmoHoverScale, const float GizmoHoverAnimationDuration, bool& bOutIsHoveringOrDraggingThisHandleGroup )
 {
 	// Call parent implementation (updates hover animation)
@@ -144,21 +213,15 @@ EGizmoHandleTypes UPivotTranslationGizmoHandleGroup::GetHandleType() const
 UPivotScaleGizmoHandleGroup::UPivotScaleGizmoHandleGroup() :
 	Super()
 {
-	UStaticMesh* TranslationHandleMesh = nullptr;
-	{
-		static ConstructorHelpers::FObjectFinder<UStaticMesh> ObjectFinder( TEXT( "/Engine/VREditor/TransformGizmo/UniformScaleHandle" ) );
-		TranslationHandleMesh = ObjectFinder.Object;
-		check( TranslationHandleMesh != nullptr );
-	}
-
-	CreateHandles( TranslationHandleMesh, FString( "PivotScaleHandle" ) );	
+	const UViewportInteractionAssetContainer& AssetContainer = UViewportWorldInteraction::LoadAssetContainer();
+	CreateHandles( AssetContainer.UniformScaleHandleMesh, FString( "PivotScaleHandle" ) );	
 }
 
-void UPivotScaleGizmoHandleGroup::UpdateGizmoHandleGroup( const FTransform& LocalToWorld, const FBox& LocalBounds, const FVector ViewLocation, bool bAllHandlesVisible, class UActorComponent* DraggingHandle, 
+void UPivotScaleGizmoHandleGroup::UpdateGizmoHandleGroup( const FTransform& LocalToWorld, const FBox& LocalBounds, const FVector ViewLocation, const bool bAllHandlesVisible, class UActorComponent* DraggingHandle, 
 	const TArray< UActorComponent* >& HoveringOverHandles, float AnimationAlpha, float GizmoScale, const float GizmoHoverScale, const float GizmoHoverAnimationDuration, bool& bOutIsHoveringOrDraggingThisHandleGroup )
 {
 	// Call parent implementation (updates hover animation)
-	Super::UpdateGizmoHandleGroup( LocalToWorld, LocalBounds, ViewLocation, bAllHandlesVisible, DraggingHandle, HoveringOverHandles, AnimationAlpha, 
+	Super::UpdateGizmoHandleGroup( LocalToWorld, LocalBounds, ViewLocation,bAllHandlesVisible, DraggingHandle, HoveringOverHandles, AnimationAlpha, 
 		GizmoScale, GizmoHoverScale, GizmoHoverAnimationDuration, bOutIsHoveringOrDraggingThisHandleGroup );
 
 	UpdateHandlesRelativeTransformOnAxis( FTransform( FVector( VREd::PivotGizmoScalePivotOffsetX->GetFloat(), 0, 0 ) ), AnimationAlpha, GizmoScale, GizmoHoverScale, ViewLocation, DraggingHandle, HoveringOverHandles );
@@ -185,17 +248,11 @@ bool UPivotScaleGizmoHandleGroup::SupportsWorldCoordinateSpace() const
 UPivotPlaneTranslationGizmoHandleGroup::UPivotPlaneTranslationGizmoHandleGroup() :
 	Super()
 {
-	UStaticMesh* TranslationHandleMesh = nullptr;
-	{
-		static ConstructorHelpers::FObjectFinder<UStaticMesh> ObjectFinder( TEXT( "/Engine/VREditor/TransformGizmo/PlaneTranslationHandle" ) );
-		TranslationHandleMesh = ObjectFinder.Object;
-		check( TranslationHandleMesh != nullptr );
-	}
-
-	CreateHandles( TranslationHandleMesh, FString( "PlaneTranslationHandle" ) );
+	const UViewportInteractionAssetContainer& AssetContainer = UViewportWorldInteraction::LoadAssetContainer();
+	CreateHandles( AssetContainer.PlaneTranslationHandleMesh, FString( "PlaneTranslationHandle" ) );
 }
 
-void UPivotPlaneTranslationGizmoHandleGroup::UpdateGizmoHandleGroup( const FTransform& LocalToWorld, const FBox& LocalBounds, const FVector ViewLocation, bool bAllHandlesVisible, class UActorComponent* DraggingHandle,
+void UPivotPlaneTranslationGizmoHandleGroup::UpdateGizmoHandleGroup( const FTransform& LocalToWorld, const FBox& LocalBounds, const FVector ViewLocation, const bool bAllHandlesVisible, class UActorComponent* DraggingHandle,
 	const TArray< UActorComponent* >& HoveringOverHandles, float AnimationAlpha, float GizmoScale, const float GizmoHoverScale, const float GizmoHoverAnimationDuration, bool& bOutIsHoveringOrDraggingThisHandleGroup )
 {
 	// Call parent implementation (updates hover animation)
@@ -222,24 +279,189 @@ EGizmoHandleTypes UPivotPlaneTranslationGizmoHandleGroup::GetHandleType() const
 UPivotRotationGizmoHandleGroup::UPivotRotationGizmoHandleGroup() :
 	Super()
 {
-	UStaticMesh* TranslationHandleMesh = nullptr;
+	const UViewportInteractionAssetContainer& AssetContainer = UViewportWorldInteraction::LoadAssetContainer();
+
+	UStaticMesh* QuarterRotationHandleMesh = AssetContainer.RotationHandleMesh;
+	CreateHandles(QuarterRotationHandleMesh, FString("RotationHandle"));
+
 	{
-		static ConstructorHelpers::FObjectFinder<UStaticMesh> ObjectFinder(TEXT("/Engine/VREditor/TransformGizmo/RotationHandleFull"));
-		TranslationHandleMesh = ObjectFinder.Object;
-		check(TranslationHandleMesh != nullptr);
+		RootFullRotationHandleComponent = CreateDefaultSubobject<USceneComponent>(TEXT("RootFullRotationHandleComponent"));
+		RootFullRotationHandleComponent->SetMobility(EComponentMobility::Movable);
+		RootFullRotationHandleComponent->SetupAttachment(this);
+
+		UStaticMesh* FullRotationHandleMesh = AssetContainer.RotationHandleSelectedMesh;
+		check(FullRotationHandleMesh != nullptr);
+
+		FullRotationHandleMeshComponent = CreateMeshHandle(FullRotationHandleMesh, FString("FullRotationHandle"));
+		FullRotationHandleMeshComponent->SetVisibility(false);
+		FullRotationHandleMeshComponent->SetCollisionEnabled(ECollisionEnabled::NoCollision);
+		FullRotationHandleMeshComponent->SetupAttachment(RootFullRotationHandleComponent);
 	}
 
-	CreateHandles(TranslationHandleMesh, FString("RotationHandle"));
+	{
+		UStaticMesh* RotationHandleIndicatorMesh = AssetContainer.StartRotationIndicatorMesh;
+		check(RotationHandleIndicatorMesh != nullptr);
+
+		//Start rotation indicator
+		RootStartRotationIdicatorComponent = CreateDefaultSubobject<USceneComponent>(TEXT("RootStartRotationIndicator"));
+		StartRotationIndicatorMeshComponent = CreateDefaultSubobject<UStaticMeshComponent>(TEXT("StartRotationIndicator"));
+		SetupIndicator(RootStartRotationIdicatorComponent, StartRotationIndicatorMeshComponent, RotationHandleIndicatorMesh);
+	}
+
+	{
+		UStaticMesh* RotationHandleIndicatorMesh = AssetContainer.CurrentRotationIndicatorMesh;
+		check(RotationHandleIndicatorMesh != nullptr);
+
+		//Delta rotation indicator
+		RootDeltaRotationIndicatorComponent = CreateDefaultSubobject<USceneComponent>(TEXT("RootDeltaRotationIndicator"));
+		DeltaRotationIndicatorMeshComponent = CreateDefaultSubobject<UStaticMeshComponent>(TEXT("DeltaRotationIndicator"));
+		SetupIndicator(RootDeltaRotationIndicatorComponent, DeltaRotationIndicatorMeshComponent, RotationHandleIndicatorMesh);
+	}
+
+	{
+		UMaterialInstanceDynamic* DynamicMaterialInst = UMaterialInstanceDynamic::Create(AssetContainer.TransformGizmoMaterial, GetTransientPackage());
+		check(DynamicMaterialInst != nullptr);
+
+		DeltaRotationIndicatorMeshComponent->SetMaterial(0, DynamicMaterialInst);
+		StartRotationIndicatorMeshComponent->SetMaterial(0, DynamicMaterialInst);
+		FullRotationHandleMeshComponent->SetMaterial(0, DynamicMaterialInst);
+
+		UMaterialInstanceDynamic* TranslucentDynamicMaterialInst = UMaterialInstanceDynamic::Create(AssetContainer.TranslucentTransformGizmoMaterial, GetTransientPackage());
+		check(TranslucentDynamicMaterialInst != nullptr);
+
+		DeltaRotationIndicatorMeshComponent->SetMaterial(1, TranslucentDynamicMaterialInst);
+		StartRotationIndicatorMeshComponent->SetMaterial(1, TranslucentDynamicMaterialInst);
+		FullRotationHandleMeshComponent->SetMaterial(1, TranslucentDynamicMaterialInst);
+
+	}
 }
 
-void UPivotRotationGizmoHandleGroup::UpdateGizmoHandleGroup( const FTransform& LocalToWorld, const FBox& LocalBounds, const FVector ViewLocation, bool bAllHandlesVisible, class UActorComponent* DraggingHandle,
+void UPivotRotationGizmoHandleGroup::UpdateGizmoHandleGroup( const FTransform& LocalToWorld, const FBox& LocalBounds, const FVector ViewLocation, const bool bAllHandlesVisible, class UActorComponent* DraggingHandle,
 	const TArray< UActorComponent* >& HoveringOverHandles, float AnimationAlpha, float GizmoScale, const float GizmoHoverScale, const float GizmoHoverAnimationDuration, bool& bOutIsHoveringOrDraggingThisHandleGroup )
 {
-	// Call parent implementation (updates hover animation)
+	// Call parent implementation
 	Super::UpdateGizmoHandleGroup(LocalToWorld, LocalBounds, ViewLocation, bAllHandlesVisible, DraggingHandle, HoveringOverHandles, AnimationAlpha,
-		GizmoScale, GizmoHoverScale, GizmoHoverAnimationDuration, bOutIsHoveringOrDraggingThisHandleGroup );
-	
-	UpdateHandlesRelativeTransformOnAxis( FTransform( FVector( 0, 0, 0 ) ), AnimationAlpha, GizmoScale, GizmoScale, ViewLocation, DraggingHandle, HoveringOverHandles );
+		GizmoScale, GizmoHoverScale, GizmoHoverAnimationDuration, bOutIsHoveringOrDraggingThisHandleGroup);
+
+	bool bShowFullRotationDragHandle = false;
+	UViewportWorldInteraction* WorldInteraction = OwningTransformGizmoActor->GetOwnerWorldInteraction();
+	const ECoordSystem CoordSystem = WorldInteraction->GetTransformGizmoCoordinateSpace();
+	const bool bIsTypeSupported = (OwningTransformGizmoActor->GetGizmoType() == EGizmoHandleTypes::All && GetShowOnUniversalGizmo()) || GetHandleType() == OwningTransformGizmoActor->GetGizmoType();
+	const bool bSupportsCurrentCoordinateSpace = SupportsWorldCoordinateSpace() || CoordSystem != COORD_World;
+	const bool bShowAnyRotationHandle = (bIsTypeSupported && bSupportsCurrentCoordinateSpace && bAllHandlesVisible);
+
+	for (int32 HandleIndex = 0; HandleIndex < Handles.Num(); ++HandleIndex)
+	{
+		FGizmoHandle& Handle = Handles[HandleIndex];
+		UStaticMeshComponent* GizmoHandleMeshComponent = Handle.HandleMesh;
+		if (GizmoHandleMeshComponent != nullptr)	// Can be null if no handle for this specific placement
+		{
+			const bool bDraggingCurrentHandle = (DraggingHandle != nullptr && DraggingHandle == GizmoHandleMeshComponent);
+			const bool bShouldShowHandle = bShowAnyRotationHandle && !bDraggingCurrentHandle;
+			GizmoHandleMeshComponent->SetVisibility(bShouldShowHandle);
+			GizmoHandleMeshComponent->SetCollisionEnabled(bShouldShowHandle ? ECollisionEnabled::QueryOnly : ECollisionEnabled::NoCollision);
+
+			int32 CenterHandleCount, FacingAxisIndex, CenterAxisIndex;
+			const FTransformGizmoHandlePlacement HandlePlacement = MakeHandlePlacementForIndex(HandleIndex);
+			HandlePlacement.GetCenterHandleCountAndFacingAxisIndex(/* Out */ CenterHandleCount, /* Out */ FacingAxisIndex, /* Out */ CenterAxisIndex);
+
+			if (bDraggingCurrentHandle)
+			{
+				bool bInitialDrag = false;
+				bShowFullRotationDragHandle = true;
+				FVector GizmoSpaceFacingAxisVector = GetAxisVector(FacingAxisIndex, HandlePlacement.Axes[FacingAxisIndex]);
+				
+				// Set the initial rotation
+				if (!StartDragRotation.IsSet())
+				{
+					StartDragRotation = LocalToWorld.GetRotation();
+					bInitialDrag = true;
+				}
+				
+				// Set the root of the full rotation handles to the rotation we had when starting the drag.
+				if (CoordSystem == ECoordSystem::COORD_Local)
+				{
+					RootFullRotationHandleComponent->SetWorldRotation(StartDragRotation.GetValue());	
+				}
+
+				FullRotationHandleMeshComponent->SetRelativeTransform(FTransform(GizmoSpaceFacingAxisVector.ToOrientationQuat(), FVector::ZeroVector, FVector(GizmoScale)));
+
+				const FVector LocalIntersectPoint = OwningTransformGizmoActor->GetOwnerWorldInteraction()->GetLocalIntersectPointOnRotationGizmo();
+				UpdateIndicator(RootDeltaRotationIndicatorComponent, LocalIntersectPoint, FacingAxisIndex);
+
+				// Update the start indicator only if this was the first time dragging
+				if (bInitialDrag)
+				{
+					UpdateIndicator(RootStartRotationIdicatorComponent, LocalIntersectPoint, FacingAxisIndex);
+
+					FLinearColor Color = WorldInteraction->GetColor(UViewportWorldInteraction::EColors::GizmoHover);
+					SetIndicatorColor(FullRotationHandleMeshComponent, Color);
+					SetIndicatorColor(DeltaRotationIndicatorMeshComponent, Color);
+					SetIndicatorColor(StartRotationIndicatorMeshComponent, Color);
+
+					bInitialDrag = false;
+				}
+			}
+			else if(DraggingHandle == nullptr)
+			{
+				int32 UpAxisIndex = 0;
+				int32 RightAxisIndex = 0;
+				FRotator Rotation;
+
+				if (FacingAxisIndex == 0)
+				{
+					// X, up = Z, right = Y
+					UpAxisIndex = 2;
+					RightAxisIndex = 1;
+					Rotation = FRotator::ZeroRotator;
+				}
+				else if (FacingAxisIndex == 1)
+				{
+					// Y, up = Z, right = X
+					UpAxisIndex = 2;
+					RightAxisIndex = 0;
+					Rotation = FRotator(0, -90, 0);
+
+				}
+				else if (FacingAxisIndex == 2)
+				{
+					// Z, up = Y, right = X
+					UpAxisIndex = 0;
+					RightAxisIndex = 1;
+					Rotation = FRotator(-90, 0, 0);
+				}
+
+				// Check on which side we are relative to the gizmo
+				const FVector GizmoSpaceViewLocation = GetOwner()->GetTransform().InverseTransformPosition(ViewLocation);
+				if (GizmoSpaceViewLocation[UpAxisIndex] < 0 && GizmoSpaceViewLocation[RightAxisIndex] < 0)
+				{
+					Rotation.Roll = 180;
+				}
+				else if (GizmoSpaceViewLocation[UpAxisIndex] < 0)
+				{
+					Rotation.Roll = 90;
+				}
+				else if (GizmoSpaceViewLocation[RightAxisIndex] < 0)
+				{
+					Rotation.Roll = -90;
+				}
+
+				const float GizmoHandleScale = GizmoScale * AnimationAlpha;
+
+				// Set the final transform
+				GizmoHandleMeshComponent->SetRelativeTransform(FTransform(Rotation, FVector::ZeroVector, FVector(GizmoHandleScale)));
+
+				// Update material
+				UpdateHandleColor(FacingAxisIndex, Handle, DraggingHandle, HoveringOverHandles);
+
+				// Reset the start rotation
+				StartDragRotation.Reset();
+			}
+		}
+	}
+
+	// Show or hide the visuals for when rotating
+	ShowRotationVisuals(bShowFullRotationDragHandle);
 }
 
 ETransformGizmoInteractionType UPivotRotationGizmoHandleGroup::GetInteractionType() const
@@ -251,3 +473,63 @@ EGizmoHandleTypes UPivotRotationGizmoHandleGroup::GetHandleType() const
 {
 	return EGizmoHandleTypes::Rotate;
 }
+
+void UPivotRotationGizmoHandleGroup::UpdateIndicator(USceneComponent* IndicatorRoot, const FVector& Direction, const uint32 FacingAxisIndex)
+{
+	float Y = 0, X = 0;
+	if (FacingAxisIndex == 0)
+	{
+		Y = Direction.Y;
+		X = Direction.Z;
+	}
+	else if (FacingAxisIndex == 1)
+	{
+		Y = -Direction.X;
+		X = Direction.Z;
+	}
+	else if (FacingAxisIndex == 2)
+	{
+		Y = Direction.Y;
+		X = -Direction.X;
+	}
+
+	const float Angle = FMath::RadiansToDegrees(FMath::Atan2(Y, X));
+	IndicatorRoot->SetRelativeRotation(FRotator(0, 0, Angle));
+}
+
+void UPivotRotationGizmoHandleGroup::ShowRotationVisuals(const bool bInShow)
+{
+	FullRotationHandleMeshComponent->SetCollisionEnabled(bInShow == true ? ECollisionEnabled::QueryOnly : ECollisionEnabled::NoCollision);
+	FullRotationHandleMeshComponent->SetVisibility(bInShow);
+	StartRotationIndicatorMeshComponent->SetVisibility(bInShow);
+	DeltaRotationIndicatorMeshComponent->SetVisibility(bInShow);
+}
+
+void UPivotRotationGizmoHandleGroup::SetupIndicator(USceneComponent* RootComponent, UStaticMeshComponent* IndicatorMeshComponent, UStaticMesh* Mesh)
+{
+	RootComponent->SetMobility(EComponentMobility::Movable);
+	RootComponent->SetupAttachment(FullRotationHandleMeshComponent);
+
+	IndicatorMeshComponent->SetStaticMesh(Mesh);
+	IndicatorMeshComponent->SetMobility(EComponentMobility::Movable);
+	IndicatorMeshComponent->SetupAttachment(RootComponent);
+	IndicatorMeshComponent->SetCollisionEnabled(ECollisionEnabled::NoCollision);
+	IndicatorMeshComponent->SetVisibility(false);
+	IndicatorMeshComponent->bGenerateOverlapEvents = false;
+	IndicatorMeshComponent->SetCanEverAffectNavigation(false);
+	IndicatorMeshComponent->bCastDynamicShadow = true;
+	IndicatorMeshComponent->bCastStaticShadow = false;
+	IndicatorMeshComponent->bAffectDistanceFieldLighting = false;
+	IndicatorMeshComponent->bAffectDynamicIndirectLighting = false;
+	IndicatorMeshComponent->SetRelativeLocation(FVector(0, 0, 100));
+}
+
+void UPivotRotationGizmoHandleGroup::SetIndicatorColor(UStaticMeshComponent* InMeshComponent, const FLinearColor& InHandleColor)
+{
+	UMaterialInstanceDynamic* MID0 = Cast<UMaterialInstanceDynamic>(InMeshComponent->GetMaterial(0));
+	UMaterialInstanceDynamic* MID1 = Cast<UMaterialInstanceDynamic>(InMeshComponent->GetMaterial(1));
+
+	static FName StaticHandleColorParameter("Color");
+	MID0->SetVectorParameterValue(StaticHandleColorParameter, InHandleColor);
+	MID1->SetVectorParameterValue(StaticHandleColorParameter, InHandleColor);
+}
\ No newline at end of file
diff --git a/Engine/Source/Editor/ViewportInteraction/Gizmo/VIPivotTransformGizmo.h b/Engine/Source/Editor/ViewportInteraction/Gizmo/VIPivotTransformGizmo.h
--- a/Engine/Source/Editor/ViewportInteraction/Gizmo/VIPivotTransformGizmo.h
+++ b/Engine/Source/Editor/ViewportInteraction/Gizmo/VIPivotTransformGizmo.h
@@ -20,12 +20,13 @@ public:
 
 	/** Default constructor that sets up CDO properties */
 	APivotTransformGizmo();
-
+	
 	/** Called by the world interaction system after we've been spawned into the world, to allow
 	    us to create components and set everything up nicely for the selected objects that we'll be
 		used to manipulate */
-	virtual void UpdateGizmo( const EGizmoHandleTypes GizmoType, const ECoordSystem GizmoCoordinateSpace, const FTransform& LocalToWorld, const FBox& LocalBounds, const FVector ViewLocation, 
-		bool bAllHandlesVisible, class UActorComponent* DraggingHandle, const TArray< UActorComponent* >& HoveringOverHandles, const float GizmoHoverScale, const float GizmoHoverAnimationDuration ) override;
+	virtual void UpdateGizmo(const EGizmoHandleTypes InGizmoType, const ECoordSystem InGizmoCoordinateSpace, const FTransform& InLocalToWorld, const FBox& InLocalBounds,
+		const FVector& InViewLocation, const float InScaleMultiplier, bool bInAllHandlesVisible, const bool bInAllowRotationAndScaleHandles, class UActorComponent* DraggingHandle,
+		const TArray<UActorComponent*>& InHoveringOverHandles, const float InGizmoHoverScale, const float InGizmoHoverAnimationDuration) override;
 
 private:
 	/** Uniform scale handle group component */
@@ -51,6 +52,13 @@ private:
 	/** Stretch handle group component */
 	UPROPERTY()
 	class UStretchGizmoHandleGroup* StretchGizmoHandleGroup;
+
+	/** The alpha for gizmo animation when aiming at it with a laser */
+	float AimingAtGizmoScaleAlpha;
+
+	/** Handle from previous tick that was dragged */
+	UPROPERTY()
+	UActorComponent* LastDraggingHandle;
 };
 
 /**
@@ -67,7 +75,7 @@ public:
 	UPivotTranslationGizmoHandleGroup();
 
 	/** Updates the gizmo handles */
-	virtual void UpdateGizmoHandleGroup( const FTransform& LocalToWorld, const FBox& LocalBounds, const FVector ViewLocation, bool bAllHandlesVisible, class UActorComponent* DraggingHandle, const TArray< UActorComponent* >& HoveringOverHandles, 
+	virtual void UpdateGizmoHandleGroup( const FTransform& LocalToWorld, const FBox& LocalBounds, const FVector ViewLocation, const bool bAllHandlesVisible, class UActorComponent* DraggingHandle, const TArray< UActorComponent* >& HoveringOverHandles, 
 		float AnimationAlpha, float GizmoScale, const float GizmoHoverScale, const float GizmoHoverAnimationDuration, bool& bOutIsHoveringOrDraggingThisHandleGroup ) override;
 
 	/** Gets the InteractionType for this Gizmo handle */
@@ -91,7 +99,7 @@ public:
 	UPivotScaleGizmoHandleGroup();
 
 	/** Updates the gizmo handles */
-	virtual void UpdateGizmoHandleGroup( const FTransform& LocalToWorld, const FBox& LocalBounds,  const FVector ViewLocation, bool bAllHandlesVisible, class UActorComponent* DraggingHandle, const TArray< UActorComponent* >& HoveringOverHandles, 
+	virtual void UpdateGizmoHandleGroup( const FTransform& LocalToWorld, const FBox& LocalBounds,  const FVector ViewLocation, const bool bAllHandlesVisible, class UActorComponent* DraggingHandle, const TArray< UActorComponent* >& HoveringOverHandles, 
 		float AnimationAlpha, float GizmoScale, const float GizmoHoverScale, const float GizmoHoverAnimationDuration, bool& bOutIsHoveringOrDraggingThisHandleGroup ) override;
 
 	/** Gets the InteractionType for this Gizmo handle */
@@ -118,7 +126,7 @@ public:
 	UPivotPlaneTranslationGizmoHandleGroup();
 
 	/** Updates the gizmo handles */
-	virtual void UpdateGizmoHandleGroup( const FTransform& LocalToWorld, const FBox& LocalBounds, const FVector ViewLocation, bool bAllHandlesVisible, class UActorComponent* DraggingHandle, const TArray< UActorComponent* >& HoveringOverHandles, 
+	virtual void UpdateGizmoHandleGroup( const FTransform& LocalToWorld, const FBox& LocalBounds, const FVector ViewLocation, const bool bAllHandlesVisible, class UActorComponent* DraggingHandle, const TArray< UActorComponent* >& HoveringOverHandles, 
 		float AnimationAlpha, float GizmoScale, const float GizmoHoverScale, const float GizmoHoverAnimationDuration, bool& bOutIsHoveringOrDraggingThisHandleGroup ) override;
 
 	/** Gets the InteractionType for this Gizmo handle */
@@ -143,7 +151,7 @@ public:
 	UPivotRotationGizmoHandleGroup();
 
 	/** Updates the gizmo handles */
-	virtual void UpdateGizmoHandleGroup(const FTransform& LocalToWorld, const FBox& LocalBounds, const FVector ViewLocation, bool bAllHandlesVisible, class UActorComponent* DraggingHandle, const TArray< UActorComponent* >& HoveringOverHandles, 
+	virtual void UpdateGizmoHandleGroup(const FTransform& LocalToWorld, const FBox& LocalBounds, const FVector ViewLocation, const bool bAllHandlesVisible, class UActorComponent* DraggingHandle, const TArray< UActorComponent* >& HoveringOverHandles, 
 		float AnimationAlpha, float GizmoScale, const float GizmoHoverScale, const float GizmoHoverAnimationDuration, bool& bOutIsHoveringOrDraggingThisHandleGroup ) override;
 
 	/** Gets the InteractionType for this Gizmo handle */
@@ -151,4 +159,43 @@ public:
 
 	/** Gets the GizmoType for this Gizmo handle */
 	virtual EGizmoHandleTypes GetHandleType() const override;
+
+private:
+
+	/** Updates the root of an indicator to rotate the indicator itself */
+	void UpdateIndicator(USceneComponent* IndicatorRoot, const FVector& Direction, const uint32 FacingAxisIndex);
+
+	/** Make the components visible when dragging rotation */
+	void ShowRotationVisuals(const bool bInShow);
+
+	void SetupIndicator(USceneComponent* RootComponent, UStaticMeshComponent* IndicatorMeshComponent, UStaticMesh* Mesh);
+	
+	void SetIndicatorColor(UStaticMeshComponent* InMeshComponent, const FLinearColor& InHandleColor);
+
+	/** Root component of all the mesh components that are used to visualize the rotation when dragging */
+	UPROPERTY()
+	USceneComponent* RootFullRotationHandleComponent;
+
+	/** When dragging a rotation handle the full rotation circle appears */
+	UPROPERTY()
+	UStaticMeshComponent* FullRotationHandleMeshComponent;
+
+	/** The mesh that indicated the start rotation */
+	UPROPERTY()
+	UStaticMeshComponent* StartRotationIndicatorMeshComponent;
+
+	/** The root component of the start rotation indicator */
+	UPROPERTY()
+	USceneComponent* RootStartRotationIdicatorComponent;
+
+	/** The mesh that indicated the delta rotation */
+	UPROPERTY()
+	UStaticMeshComponent* DeltaRotationIndicatorMeshComponent;
+
+	/** The root component of the delta rotation indicator */
+	UPROPERTY()
+	USceneComponent* RootDeltaRotationIndicatorComponent;
+
+	/** The rotation when starting to drag the gizmo */
+	TOptional<FQuat> StartDragRotation;
 };
diff --git a/Engine/Source/Editor/ViewportInteraction/Gizmo/VIUniformScaleGizmoHandle.cpp b/Engine/Source/Editor/ViewportInteraction/Gizmo/VIUniformScaleGizmoHandle.cpp
--- a/Engine/Source/Editor/ViewportInteraction/Gizmo/VIUniformScaleGizmoHandle.cpp
+++ b/Engine/Source/Editor/ViewportInteraction/Gizmo/VIUniformScaleGizmoHandle.cpp
@@ -27,7 +27,7 @@ UUniformScaleGizmoHandleGroup::UUniformScaleGizmoHandleGroup()
 	NewHandle.HandleMesh = UniformScaleHandle;
 }
 
-void UUniformScaleGizmoHandleGroup::UpdateGizmoHandleGroup( const FTransform& LocalToWorld, const FBox& LocalBounds, const FVector ViewLocation, bool bAllHandlesVisible, class UActorComponent* DraggingHandle, const TArray< UActorComponent* >& HoveringOverHandles, 
+void UUniformScaleGizmoHandleGroup::UpdateGizmoHandleGroup( const FTransform& LocalToWorld, const FBox& LocalBounds, const FVector ViewLocation, const bool bAllHandlesVisible, class UActorComponent* DraggingHandle, const TArray< UActorComponent* >& HoveringOverHandles, 
 	float AnimationAlpha, float GizmoScale, const float GizmoHoverScale, const float GizmoHoverAnimationDuration, bool& bOutIsHoveringOrDraggingThisHandleGroup )
 {
 	// Call parent implementation (updates hover animation)
@@ -70,16 +70,16 @@ void UUniformScaleGizmoHandleGroup::UpdateGizmoHandleGroup( const FTransform& Lo
 			if( GizmoActor )
 			{
 				UViewportWorldInteraction* WorldInteraction = GizmoActor->GetOwnerWorldInteraction();
-				if( WorldInteraction )
+				if( WorldInteraction && WorldInteraction->IsActive() )
 				{
 					FLinearColor HandleColor = WorldInteraction->GetColor( UViewportWorldInteraction::EColors::DefaultColor );
 					if (UniformScaleHandle == DraggingHandle)
 					{
-						HandleColor = WorldInteraction->GetColor( UViewportWorldInteraction::EColors::Dragging );
+						HandleColor = WorldInteraction->GetColor( UViewportWorldInteraction::EColors::GizmoDragging );
 					}
 					else if (HoveringOverHandles.Contains( UniformScaleHandle ))
 					{
-						HandleColor = FLinearColor::LerpUsingHSV( HandleColor, WorldInteraction->GetColor( UViewportWorldInteraction::EColors::Hover ), Handle.HoverAlpha );
+						HandleColor = FLinearColor::LerpUsingHSV( HandleColor, WorldInteraction->GetColor( UViewportWorldInteraction::EColors::GizmoHover ), Handle.HoverAlpha );
 					}
 
 					MID0->SetVectorParameterValue( "Color", HandleColor );
diff --git a/Engine/Source/Editor/ViewportInteraction/Public/ActorTransformer.h b/Engine/Source/Editor/ViewportInteraction/Public/ActorTransformer.h
new file mode 100644
--- /dev/null
+++ b/Engine/Source/Editor/ViewportInteraction/Public/ActorTransformer.h
@@ -0,0 +1,35 @@
+// Copyright 1998-2016 Epic Games, Inc. All Rights Reserved.
+
+#pragma once
+
+#include "ViewportTransformer.h"
+#include "ActorTransformer.generated.h"
+
+
+UCLASS()
+class VIEWPORTINTERACTION_API UActorTransformer : public UViewportTransformer
+{
+	GENERATED_BODY()
+
+public:
+
+	// UViewportTransformer overrides
+	virtual void Init( class UViewportWorldInteraction* InitViewportWorldInteraction ) override;
+	virtual void Shutdown() override;
+	virtual bool CanAlignToActors() const override
+	{
+		return true;
+	}
+	virtual void OnStartDragging(class UViewportInteractor* Interactor) override;
+	virtual void OnStopDragging(class UViewportInteractor* Interactor) override;
+
+protected:
+
+	/** Called when selection changes in the level */
+	void OnActorSelectionChanged( UObject* ChangedObject );
+
+	/** Creates a list of transformables with the current selected actors */
+	void UpdateTransformables();
+
+};
+
diff --git a/Engine/Source/Editor/ViewportInteraction/Public/MouseCursorInteractor.h b/Engine/Source/Editor/ViewportInteraction/Public/MouseCursorInteractor.h
--- a/Engine/Source/Editor/ViewportInteraction/Public/MouseCursorInteractor.h
+++ b/Engine/Source/Editor/ViewportInteraction/Public/MouseCursorInteractor.h
@@ -22,8 +22,15 @@ public:
 
 	// ViewportInteractor overrides
 	virtual void PollInput() override;
+	virtual bool IsModifierPressed() const override;
+
+protected:
+
+	virtual bool AllowLaserSmoothing() const;
 
 private:
 
-	// ...
+	/** Whether the control key was pressed the last time input was polled */
+	bool bIsControlKeyPressed;
+
 };
diff --git a/Engine/Source/Editor/ViewportInteraction/Public/ViewportDragOperation.h b/Engine/Source/Editor/ViewportInteraction/Public/ViewportDragOperation.h
--- a/Engine/Source/Editor/ViewportInteraction/Public/ViewportDragOperation.h
+++ b/Engine/Source/Editor/ViewportInteraction/Public/ViewportDragOperation.h
@@ -52,6 +52,9 @@ public:
 	/* Destroys the dragoperation */
 	void ClearDragOperation();
 
+	/** If this operation is currently active. */
+	bool IsDragging() const;
+
 private:
 
 	/** The actual dragging implementation */
diff --git a/Engine/Source/Editor/ViewportInteraction/Public/ViewportInteractionUtils.h b/Engine/Source/Editor/ViewportInteraction/Public/ViewportInteractionUtils.h
new file mode 100644
--- /dev/null
+++ b/Engine/Source/Editor/ViewportInteraction/Public/ViewportInteractionUtils.h
@@ -0,0 +1,73 @@
+// Copyright 1998-2017 Epic Games, Inc. All Rights Reserved.
+
+#pragma once
+
+#include "CoreMinimal.h"
+
+namespace ViewportInteractionUtils
+{
+	/************************************************************************/
+	/* 1 Euro filter smoothing algorithm									*/
+	/* http://cristal.univ-lille.fr/~casiez/1euro/							*/
+	/************************************************************************/
+	class VIEWPORTINTERACTION_API FOneEuroFilter
+	{
+	private:
+
+		class FLowpassFilter
+		{
+		public:
+
+			/** Default constructor */
+			FLowpassFilter();
+
+			/** Calculate */
+			FVector Filter(const FVector& InValue, const FVector& InAlpha);
+
+			/** If the filter was not executed yet */
+			bool IsFirstTime() const;
+
+			/** Get the previous filtered value */
+			FVector GetPrevious() const;
+
+		private:
+
+			/** The previous filtered value */
+			FVector Previous;
+
+			/** If this is the first time doing a filter */
+			bool bFirstTime;
+		};
+
+	public:
+
+		/** Default constructor */
+		FOneEuroFilter();
+
+		FOneEuroFilter(const double InMinCutoff, const double InCutoffSlope, const double InDeltaCutoff);
+
+		/** Smooth vector */
+		FVector Filter(const FVector& InRaw, const double InDeltaTime);
+
+		/** Set the minimum cutoff */
+		void SetMinCutoff(const double InMinCutoff);
+
+		/** Set the cutoff slope */
+		void SetCutoffSlope(const double InCutoffSlope);
+
+		/** Set the delta slope */
+		void SetDeltaCutoff(const double InDeltaCutoff);
+
+	private:
+
+		const FVector CalculateCutoff(const FVector& InValue);
+		const FVector CalculateAlpha(const FVector& InCutoff, const double InDeltaTime) const;
+		const float CalculateAlpha(const float InCutoff, const double InDeltaTime) const;
+
+		double MinCutoff;
+		double CutoffSlope;
+		double DeltaCutoff;
+		FLowpassFilter RawFilter;
+		FLowpassFilter DeltaFilter;
+	};
+}
\ No newline at end of file
diff --git a/Engine/Source/Editor/ViewportInteraction/Public/ViewportInteractor.h b/Engine/Source/Editor/ViewportInteraction/Public/ViewportInteractor.h
--- a/Engine/Source/Editor/ViewportInteraction/Public/ViewportInteractor.h
+++ b/Engine/Source/Editor/ViewportInteraction/Public/ViewportInteractor.h
@@ -9,6 +9,7 @@
 #include "Engine/EngineTypes.h"
 #include "ViewportInteractionTypes.h"
 #include "Editor/ViewportInteraction/ViewportInteractorData.h"
+#include "ViewportInteractionUtils.h"
 #include "ViewportInteractor.generated.h"
 
 class AActor;
@@ -26,6 +27,9 @@ public:
 	/** Gets the private data for this interactor */
 	struct FViewportInteractorData& GetInteractorData();
 
+	/** Gets the private data for this interactor (const) */
+	const struct FViewportInteractorData& GetInteractorData() const;
+
 	/** Sets the world interaction */
 	void SetWorldInteraction( class UViewportWorldInteraction* InWorldInteraction );
 
@@ -38,17 +42,14 @@ public:
 	/** Gets the paired interactor assigned by the world interaction, can return null when there is no other interactor */
 	UViewportInteractor* GetOtherInteractor() const;
 
-	/** Updates the hover state of this interactor */
-	void UpdateHoverResult( FHitResult& HitResult );
-
 	/** Whenever the ViewportWorldInteraction is shut down, the interacts will shut down as well */
 	virtual void Shutdown();
 
 	/** Update for this interactor called by the ViewportWorldInteraction */
-	virtual void Tick( const float DeltaTime ) {};
+	virtual void Tick( const float DeltaTime );
 
-	/** Gets the hover component from the interactor data */
-	virtual class UActorComponent* GetHoverComponent();
+	/** Gets the last component hovered on by the interactor laser. */
+	virtual class UActorComponent* GetLastHoverComponent();
 
 	/** Adds a new action to the KeyToActionMap */
 	void AddKeyAction( const FKey& Key, const FViewportActionKeyInput& Action );
@@ -60,10 +61,16 @@ public:
 	virtual void PollInput() {}
 
 	/** Handles axis input and translates it actions */
-	bool HandleInputKey( const FKey Key, const EInputEvent Event );
+	bool HandleInputKey( class FEditorViewportClient& ViewportClient, const FKey Key, const EInputEvent Event );
 
 	/** Handles axis input and translates it to actions */
-	bool HandleInputAxis( const FKey Key, const float Delta, const float DeltaTime );
+	bool HandleInputAxis( class FEditorViewportClient& ViewportClient, const FKey Key, const float Delta, const float DeltaTime );
+
+	/** @return	Returns true if this interactor's designated 'modifier button' is currently held down.  Some interactors may not support this */
+	virtual bool IsModifierPressed() const
+	{
+		return false;
+	}
 
 	/** Gets the world transform of this interactor */
 	FTransform GetTransform() const;
@@ -77,19 +84,32 @@ public:
 	/** Gets the interactor data drag velocity */
 	FVector GetDragTranslationVelocity() const;
 
+	/** Sets the hover location */
+	void SetHoverLocation(const FVector& InHoverLocation);
+
 	/**
 	 * Gets the start and end point of the laser pointer for the specified hand
 	 *
 	 * @param HandIndex				Index of the hand to use
 	 * @param LasertPointerStart	(Out) The start location of the laser pointer in world space
 	 * @param LasertPointerEnd		(Out) The end location of the laser pointer in world space
-	 * @param bEvenIfUIIsInFront	If true, returns a laser pointer even if the hand has UI in front of it (defaults to false)
+	 * @param bEvenIfBlocked		If true, returns a laser pointer even if the hand has UI in front of it (defaults to false)
 	 * @param LaserLengthOverride	If zero the default laser length (VREdMode::GetLaserLength) is used
 	 *
 	 * @return	True if we have motion controller data for this hand and could return a valid result
 	 */
 	bool GetLaserPointer( FVector& LaserPointerStart, FVector& LaserPointerEnd, const bool bEvenIfBlocked = false, const float LaserLengthOverride = 0.0f );
 
+	/**
+	 * Gets a sphere on this interactor that can be used to interact with objects in close proximity
+	 *
+	 * @param	OutGrabberSphere	The sphere in world space
+	 * @param	bEvenIfBlocked	When set to true, a valid sphere may be returned even if there is UI attached in front of this interactor
+	 *
+	 * @return	True if the sphere is available and valid, or false is the interactor was busy or we could not determine a valid position for the interactor
+	 */
+	bool GetGrabberSphere( FSphere& OutGrabberSphere, const bool bEvenIfBlocked = false );
+
 	/** Gets the maximum length of a laser pointer */
 	float GetLaserPointerMaxLength() const;
 
@@ -99,7 +119,6 @@ public:
 	/**
 	 * Traces along the laser pointer vector and returns what it first hits in the world
 	 *
-	 * @param HandIndex	Index of the hand that uses has the laser pointer to trace
 	 * @param OptionalListOfIgnoredActors Actors to exclude from hit testing
 	 * @param bIgnoreGizmos True if no gizmo results should be returned, otherwise they are preferred (x-ray)
 	 * @param bEvenIfUIIsInFront If true, ignores any UI that might be blocking the ray
@@ -111,10 +130,10 @@ public:
 		TArray<UClass*>* ObjectsInFrontOfGizmo = nullptr, const bool bEvenIfBlocked = false, const float LaserLengthOverride = 0.0f );
 
 	/** Reset the values before checking the hover actions */
-	virtual void ResetHoverState( const float DeltaTime ) {};
+	virtual void ResetHoverState();
 
 	/** Needs to be implemented by the base to calculate drag ray length and the velocity for the ray */
-	virtual void CalculateDragRay() {};
+	virtual void CalculateDragRay( float& InOutDragRayLength, float& InOutDragRayVelocity ) {};
 
 	/**
 	 * Creates a hand transform and forward vector for a laser pointer for a given hand
@@ -125,13 +144,13 @@ public:
 	 *
 	 * @return	True if we have motion controller data for this hand and could return a valid result
 	 */
-	virtual bool GetTransformAndForwardVector( FTransform& OutHandTransform, FVector& OutForwardVector );
+	virtual bool GetTransformAndForwardVector( FTransform& OutHandTransform, FVector& OutForwardVector ) const;
 
-	/** Called by StartDragging in world interaction to give the interator a chance of acting upon starting a drag operation */
-	virtual void OnStartDragging( UActorComponent* ClickedComponent, const FVector& HitLocation, const bool bIsPlacingActors ) {};
+	/** Called by StartDragging in world interaction to give the interactor a chance of acting upon starting a drag operation */
+	virtual void OnStartDragging( const FVector& HitLocation, const bool bIsPlacingNewObjects );
 
 	/** Gets the interactor laser hover location */
-	FVector GetHoverLocation() const;
+	FVector GetHoverLocation();
 
 	/** If the interactor laser is currently hovering */
 	bool IsHovering() const;
@@ -143,7 +162,7 @@ public:
 	void SetDraggingMode( const EViewportInteractionDraggingMode NewDraggingMode );
 
 	/** To be overridden by base class. Called by GetLaserPointer to give the derived interactor a chance to disable the laser. By default it is not blocked */
-	virtual bool GetIsLaserBlocked() { return false; }
+	virtual bool GetIsLaserBlocked() const { return false; }
 
 	/** Gets a certain action by iterating through the map looking for the same ActionType */
 	// @todo ViewportInteractor: This should be changed to return a const pointer, but we need to fix up some dragging code in WorldInteraction first
@@ -152,41 +171,31 @@ public:
 	/** Gets the drag haptic feedback strength console variable */
 	float GetDragHapticFeedbackStrength() const;
 
-	/**
-	 * Enables or disable "light press" locking for this interactor.  This is an advanced feature.
-	 * It should be true if we allow locking of the lightly pressed state for the current event.  This can 
-	 * be useful to turn off in cases where you always want a full press to override the light press, even 
-	 * if it was a very slow press
-	 *
-	 * @param	bAllow		Enables or disables light press locking for this hand
-	 */
-	void SetAllowTriggerLightPressLocking( const bool bInAllow );
+	/** If this interactor is hovering over a type that has priority from GetHitResultFromLaserPointer */
+	bool IsHoveringOverPriorityType() const;
 
-	/** @return Check if this interactor allows the light press to be locked */
-	bool AllowTriggerLightPressLocking() const;
+	/** Returns true if currently hovering over selected actor */
+	bool IsHoveringOverSelectedActor() const;
 
-	/**
-	 * When responding to a "light press" event, you can call this and pass false to to prevent a full press event from
-	 * being possible at all.  Useful when you're not planning to use FullyPressed for anything, but you don't want a 
-	 * full press to cancel your light press.
-	 *
-	 * @param	bInAllow	Whether to allow a full press to interrupt the light press.  Defaults to true, and will reset to true with every re-press of the trigger.
-	 */
-	void SetAllowTriggerFullPress( const bool bInAllow );
+	/** Reset the stored laser end location at the end of tick */
+	void ResetLaserEnd();
 
-	/** @return Returns whether a full press is allowed to interrupt a light press. */
-	bool AllowTriggerFullPress() const;
+protected:
 
-	/** If this interactor is hovering over a type that has priority from GetHitResultFromLaserPointer */
-	bool IsHoveringOverPriorityType() const;
+	/** To be overridden. Called by HandleInputKey before delegates and default input implementation */
+	virtual void PreviewInputKey( class FEditorViewportClient& ViewportClient, FViewportActionKeyInput& Action, const FKey Key, const EInputEvent Event, bool& bOutWasHandled ) {};
 
-protected:
+	/** To be overridden. Called by HandleInputAxis before delegates and default input implementation */
+	virtual void PreviewInputAxis( class FEditorViewportClient& ViewportClient, FViewportActionKeyInput& Action, const FKey Key, const float Delta, const float DeltaTime, bool& bOutWasHandled ) {};
 
-	/** To be overridden by base class. Called by HandleInputKey before delegates and default input implementation */
-	virtual void HandleInputKey( FViewportActionKeyInput& Action, const FKey Key, const EInputEvent Event, bool& bOutWasHandled ) {};
+	/** To be overridden. Called by HandleInputKey before delegates and default input implementation */
+	virtual void HandleInputKey( class FEditorViewportClient& ViewportClient, FViewportActionKeyInput& Action, const FKey Key, const EInputEvent Event, bool& bOutWasHandled ) {};
 	
-	/** To be overridden by base class. Called by HandleInputAxis before delegates and default input implementation */
-	virtual void HandleInputAxis( FViewportActionKeyInput& Action, const FKey Key, const float Delta, const float DeltaTime, bool& bOutWasHandled ) {};
+	/** To be overridden. Called by HandleInputAxis before delegates and default input implementation */
+	virtual void HandleInputAxis( class FEditorViewportClient& ViewportClient, FViewportActionKeyInput& Action, const FKey Key, const float Delta, const float DeltaTime, bool& bOutWasHandled ) {};
+
+	/** If this interactors allows to smooth the laser. Default is true. */
+	virtual bool AllowLaserSmoothing() const;
 
 protected:
 
@@ -203,4 +212,15 @@ protected:
 	/** The paired interactor by the world interaction */
 	UPROPERTY()
 	UViewportInteractor* OtherInteractor;
-};
+
+	/** True if this interactor supports 'grabber sphere' interaction.  Usually disabled for mouse cursors */
+	bool bAllowGrabberSphere;
+
+	/** Store end of the laser pointer. This will be returned when calling GetLaserPointer multiple times a tick */
+	TOptional<FVector> SavedLaserPointerEnd;
+
+private:
+
+	/** Smoothing filter for laser */
+	ViewportInteractionUtils::FOneEuroFilter SmoothingOneEuroFilter;
+};
\ No newline at end of file
