Branch snapshot for CL 3332592

[CL 3332592 in Dev-Geometry branch]

diff --git a/Engine/Source/Editor/VREditor/Public/VREditorAssetContainer.h b/Engine/Source/Editor/VREditor/Public/VREditorAssetContainer.h
--- a/Engine/Source/Editor/VREditor/Public/VREditorAssetContainer.h
+++ b/Engine/Source/Editor/VREditor/Public/VREditorAssetContainer.h
@@ -2,81 +2,150 @@
 
 #pragma once
 
-#include "SlateWidgetStyleContainerBase.h"
-#include "SlateWidgetStyle.h"
-#include "Sound/SoundBase.h"
+#include "Engine/DataAsset.h"
 #include "VREditorAssetContainer.generated.h"
 
+
+// Forward declarations
+class USoundBase;
+class USoundCue;
+class UStaticMesh;
+class UMaterial;
+class UMaterialInterface;
+class UMaterialInstance;
+class UFont;
+
 /**
- * Represents the common menu sounds used in viewport interaction
+ * Asset container for VREditor.
  */
-USTRUCT()
-struct FVREditorAssetContainer : public FSlateWidgetStyle
+UCLASS()
+class VREDITOR_API UVREditorAssetContainer : public UDataAsset
 {
-	GENERATED_USTRUCT_BODY()
-
-	FVREditorAssetContainer();
-	virtual ~FVREditorAssetContainer();
+	GENERATED_BODY()
 
-	// FSlateWidgetStyle
-	virtual void GetResources(TArray<const FSlateBrush*>& OutBrushes) const override;
-	static const FName TypeName;
-	virtual const FName GetTypeName() const override { return TypeName; };
-	static const FVREditorAssetContainer& GetDefault();
+public:
+	
+	//
+	// Sounds
+	//
 
-	/** The sound that should play when starting the game */
 	UPROPERTY(EditAnywhere, Category = Sound)
 	USoundBase* DockableWindowCloseSound;
-	FVREditorAssetContainer& SetDockableWindowCloseSound(USoundBase* InSound) { DockableWindowCloseSound = InSound; return *this; }
 
 	UPROPERTY(EditAnywhere, Category = Sound)
 	USoundBase* DockableWindowOpenSound;
-	FVREditorAssetContainer& SetDockableWindowOpenSound(USoundBase* InSound) { DockableWindowOpenSound = InSound; return *this; }
 
 	UPROPERTY(EditAnywhere, Category = Sound)
 	USoundBase* DockableWindowDropSound;
-	FVREditorAssetContainer& SetDockableWindowDropSound(USoundBase* InSound) { DockableWindowDropSound = InSound; return *this; }
 
 	UPROPERTY(EditAnywhere, Category = Sound)
 	USoundBase* DockableWindowDragSound;
-	FVREditorAssetContainer& SetDockableWindowDragSound(USoundBase* InSound) { DockableWindowDragSound = InSound; return *this; }
 
 	UPROPERTY(EditAnywhere, Category = Sound)
 	USoundBase* DropFromContentBrowserSound;
-	FVREditorAssetContainer& SetDropFromContentBrowserSound(USoundBase* InSound) { DropFromContentBrowserSound = InSound; return *this; }
 
 	UPROPERTY(EditAnywhere, Category = Sound)
 	USoundBase* RadialMenuOpenSound;
-	FVREditorAssetContainer& SetRadialMenuOpenSound(USoundBase* InSound) { RadialMenuOpenSound = InSound; return *this; }
 
 	UPROPERTY(EditAnywhere, Category = Sound)
 	USoundBase* RadialMenuCloseSound;
-	FVREditorAssetContainer& SetRadialMenuCloseSound(USoundBase* InSound) { RadialMenuCloseSound = InSound; return *this; }
 
 	UPROPERTY(EditAnywhere, Category = Sound)
 	USoundBase* TeleportSound;
-	FVREditorAssetContainer& SetTeleportSound(USoundBase* InSound) { TeleportSound = InSound; return *this; }
 
 	UPROPERTY(EditAnywhere, Category = Sound)
 	USoundCue* ButtonPressSound;
-	FVREditorAssetContainer& SetTeleportSound(USoundCue* InSound) { ButtonPressSound = InSound; return *this; }
-};
 
+	//
+	// Meshes
+	//
 
-/**
- */
-UCLASS(hidecategories=Object, MinimalAPI)
-class UVREditorStyleContainer : public USlateWidgetStyleContainerBase
-{
-	GENERATED_UCLASS_BODY()
+	UPROPERTY(EditAnywhere, Category = Mesh)
+	UStaticMesh* GenericHMDMesh;
 
-public:
-	/** The actual data describing the sounds */
-	UPROPERTY(Category=Appearance, EditAnywhere, meta=(ShowOnlyInnerProperties))
-	FVREditorAssetContainer SoundsStyle;
-
-	virtual const struct FSlateWidgetStyle* const GetStyle() const override
-	{
-		return static_cast< const struct FSlateWidgetStyle* >( &SoundsStyle );
-	}
+	UPROPERTY(EditAnywhere, Category = Mesh)
+	UStaticMesh* PlaneMesh;
+
+	UPROPERTY(EditAnywhere, Category = Mesh)
+	UStaticMesh* CylinderMesh;
+
+	UPROPERTY(EditAnywhere, Category = Mesh)
+	UStaticMesh* LaserPointerStartMesh;
+
+	UPROPERTY(EditAnywhere, Category = Mesh)
+	UStaticMesh* LaserPointerMesh;
+
+	UPROPERTY(EditAnywhere, Category = Mesh)
+	UStaticMesh* LaserPointerEndMesh;
+
+	UPROPERTY(EditAnywhere, Category = Mesh)
+	UStaticMesh* LaserPointerHoverMesh;
+
+	UPROPERTY(EditAnywhere, Category = Mesh)
+	UStaticMesh* VivePreControllerMesh;
+
+	UPROPERTY(EditAnywhere, Category = Mesh)
+	UStaticMesh* OculusControllerMesh;
+
+	UPROPERTY(EditAnywhere, Category = Mesh)
+	UStaticMesh* GenericControllerMesh;
+
+	UPROPERTY(EditAnywhere, Category = Mesh)
+	UStaticMesh* TeleportRootMesh;
+
+	UPROPERTY(EditAnywhere, Category = Mesh)
+	UStaticMesh* WindowMesh;
+
+	UPROPERTY(EditAnywhere, Category = Mesh)
+	UStaticMesh* WindowSelectionBarMesh;
+
+	UPROPERTY(EditAnywhere, Category = Mesh)
+	UStaticMesh* WindowCloseButtonMesh;
+
+	UPROPERTY(EditAnywhere, Category = Mesh)
+	UStaticMesh* RadialMenuMainMesh;
+
+	UPROPERTY(EditAnywhere, Category = Mesh)
+	UStaticMesh* RadialMenuPointerMesh;
+	
+	//
+	// Materials
+	//
+
+	UPROPERTY(EditAnywhere, Category = Material)
+	UMaterialInterface* GridMaterial;
+
+	UPROPERTY(EditAnywhere, Category = Material)
+	UMaterialInterface* LaserPointerMaterial;
+
+	UPROPERTY(EditAnywhere, Category = Material)
+	UMaterialInterface* LaserPointerTranslucentMaterial;
+
+	UPROPERTY(EditAnywhere, Category = Material)
+	UMaterial* WorldMovementPostProcessMaterial;
+
+	UPROPERTY(EditAnywhere, Category = Material)
+	UMaterialInterface* TextMaterial;
+	
+	UPROPERTY(EditAnywhere, Category = Material)
+	UMaterialInterface* VivePreControllerMaterial;
+
+	UPROPERTY(EditAnywhere, Category = Material)
+	UMaterialInterface* OculusControllerMaterial;
+
+	UPROPERTY(EditAnywhere, Category = Material)
+	UMaterialInterface* TeleportMaterial;
+
+	UPROPERTY(EditAnywhere, Category = Material)
+	UMaterialInterface* WindowMaterial;
+
+	UPROPERTY(EditAnywhere, Category = Material)
+	UMaterialInterface* WindowTranslucentMaterial;
+
+	//
+	// Fonts
+	//
+
+	UPROPERTY(EditAnywhere, Category = Font)
+	UFont* TextFont;
 };
diff --git a/Engine/Source/Editor/VREditor/Public/VREditorMode.h b/Engine/Source/Editor/VREditor/Public/VREditorMode.h
--- a/Engine/Source/Editor/VREditor/Public/VREditorMode.h
+++ b/Engine/Source/Editor/VREditor/Public/VREditorMode.h
@@ -198,9 +198,9 @@ public:
 	/** @return Returns the type of HMD we're dealing with */
 	EHMDDeviceType::Type GetHMDDeviceType() const;
 
-	/** @return Checks to see if the specified hand is aiming roughly toward the specified capsule */
-	bool IsHandAimingTowardsCapsule( class UViewportInteractor* Interactor, const FTransform& CapsuleTransform, const FVector CapsuleStart, const FVector CapsuleEnd, const float CapsuleRadius, const float MinDistanceToCapsule, const FVector CapsuleFrontDirection, const float MinDotForAimingAtCapsule ) const;
-	
+	/** @return Checks to see if the specified interactor is aiming roughly toward the specified capsule */
+	bool IsHandAimingTowardsCapsule(class UViewportInteractor* Interactor, const FTransform& CapsuleTransform, const FVector CapsuleStart, const FVector CapsuleEnd, const float CapsuleRadius, const float MinDistanceToCapsule, const FVector CapsuleFrontDirection, const float MinDotForAimingAtCapsule) const;
+
 	/** Gets the hand interactor  */
 	class UVREditorInteractor* GetHandInteractor( const EControllerHand ControllerHand ) const;
 
@@ -268,8 +268,9 @@ public:
 	/** Create a static motion controller mesh for the current HMD platform */
 	UStaticMeshComponent* CreateMotionControllerMesh( AActor* OwningActor, USceneComponent* AttachmentToComponent );
 
-	/** Helper function to create a static mesh */
+	/** Helper functions to create a static mesh */
 	UStaticMeshComponent* CreateMesh( AActor* OwningActor, const FString& MeshName, USceneComponent* AttachmentToComponent /*= nullptr */ );
+	UStaticMeshComponent* CreateMesh(AActor* OwningActor, UStaticMesh* Mesh, USceneComponent* AttachmentToComponent /*= nullptr */);
 
 	/** Sets a delegate for the context-specific actions menu */
 	void SetActionsMenuGenerator(const FOnRadialMenuGenerated NewMenuGenerator, const FText NewLabel);
@@ -287,7 +288,7 @@ public:
 	bool GetStartedPlayFromVREditor() const;
 	
 	/** Gets the container for all the assets of VREditor. */
-	const struct FVREditorAssetContainer& GetAssetContainer() const;
+	const class UVREditorAssetContainer& GetAssetContainer() const;
 	
 	/** Plays sound at location. */
 	void PlaySound(USoundBase* SoundBase, const FVector& InWorldLocation, const float InVolume = 1.0f);
@@ -458,6 +459,9 @@ public:
 	/** Returns the currently active sequencer */
 	class ISequencer* GetCurrentSequencer();
 
+	/** The asset container path */
+	static const FString AssetContainerPath;
+
 private:
 
 	// All the colors for this mode
@@ -479,5 +483,6 @@ private:
 	bool bStartedPlayFromVREditor;	
 	
 	/** Container of assets */
-	TSharedPtr<const struct FVREditorAssetContainer> AssetContainer;
+	UPROPERTY()
+	class UVREditorAssetContainer* AssetContainer;
 };
diff --git a/Engine/Source/Editor/VREditor/Teleporter/VREditorTeleporter.cpp b/Engine/Source/Editor/VREditor/Teleporter/VREditorTeleporter.cpp
--- a/Engine/Source/Editor/VREditor/Teleporter/VREditorTeleporter.cpp
+++ b/Engine/Source/Editor/VREditor/Teleporter/VREditorTeleporter.cpp
@@ -23,11 +23,12 @@ namespace VREd
 	static FAutoConsoleVariable TeleportDistance(TEXT("VREd.TeleportDistance"), 500.0f, TEXT("Default distance for teleporting when not hitting anything"));
 	static FAutoConsoleVariable TeleportScaleSensitivity(TEXT("VREd.TeleportScaleSensitivity"), 0.05f, TEXT("Teleport world to meters scale touchpad sensitivity"));
 	static FAutoConsoleVariable TeleportOffsetMultiplier(TEXT("VREd.TeleportOffsetMultiplier"), 0.3f, TEXT("Teleport offset multiplier"));
-	static FAutoConsoleVariable TeleportEnableChangeScale(TEXT("VREd.TeleportEnableChangeScale"), 0, TEXT("Ability to change the world to meters scale while teleporting"));
+	static FAutoConsoleVariable TeleportEnableChangeScale(TEXT("VREd.TeleportEnableChangeScale"), 1, TEXT("Ability to change the world to meters scale while teleporting"));
 	static FAutoConsoleVariable TeleportFadeInAnimateSpeed(TEXT("VREd.TeleportAnimateSpeed"), 3.0f, TEXT("How fast the teleporter should fade in"));
 	static FAutoConsoleVariable TeleportDragSpeed(TEXT("VREd.TeleportDragSpeed"), 0.3f, TEXT("How fast the teleporter should drag behind the laser aiming location"));
 	static FAutoConsoleVariable TeleportAllowScaleBackToDefault(TEXT("VREd.TeleportAllowScaleBackToDefault"), 1, TEXT("Scale back to default world to meters scale"));
-	static FAutoConsoleVariable TeleportAllowPushPull(TEXT("VREd.TeleportAllowPushPull"), 0, TEXT("Allow being able to push and pull the teleporter along the laser."));
+	static FAutoConsoleVariable TeleportAllowPushPull(TEXT("VREd.TeleportAllowPushPull"), 1, TEXT("Allow being able to push and pull the teleporter along the laser."));
+	static FAutoConsoleVariable TeleportSlideBuffer(TEXT("VREd.TeleportSlideBuffer"), 0.3f, TEXT("The minimum slide on trackpad to push/pull or change scale."));
 }
 
 AVREditorTeleporter::AVREditorTeleporter():
@@ -76,15 +77,17 @@ void AVREditorTeleporter::Init(UVREditorMode* InMode)
 
 	SetActorEnableCollision(false);
 
+	const UVREditorAssetContainer& AssetContainer = VRMode->GetAssetContainer();
+
 	{
-		UMaterialInstance* RoomSpaceMaterialInstance = LoadObject<UMaterialInstance>(nullptr, TEXT("/Engine/VREditor/Teleport/TeleportDirectionMaterial_Inst"));
-		check(RoomSpaceMaterialInstance != nullptr);
-		TeleportMID = UMaterialInstanceDynamic::Create(RoomSpaceMaterialInstance, GetTransientPackage());
+		UMaterialInterface* RoomSpaceMaterial = AssetContainer.TeleportMaterial;
+		check(RoomSpaceMaterial != nullptr);
+		TeleportMID = UMaterialInstanceDynamic::Create(RoomSpaceMaterial, GetTransientPackage());
 		check(TeleportMID != nullptr);
 	}
 
 	{
-		TeleportDirectionMeshComponent = VRMode->CreateMesh(this, FString("/Engine/VREditor/Teleport/TeleportDirectionMesh"), GetRootComponent());
+		TeleportDirectionMeshComponent = VRMode->CreateMesh(this, AssetContainer.TeleportRootMesh, GetRootComponent());
 		TeleportDirectionMeshComponent->SetWorldScale3D(FVector(1, 1, 1));
 		TeleportDirectionMeshComponent->SetMaterial(0, TeleportMID);
 		TeleportDirectionMeshComponent->SetCollisionEnabled(ECollisionEnabled::NoCollision);
@@ -92,7 +95,7 @@ void AVREditorTeleporter::Init(UVREditorMode* InMode)
 		TeleportDirectionMeshComponent->SetCollisionResponseToChannel(COLLISION_GIZMO, ECollisionResponse::ECR_Ignore);
 		TeleportDirectionMeshComponent->SetCastShadow(false);
 
-		HMDMeshComponent = VRMode->CreateMesh(this, FString("/Engine/VREditor/Devices/Generic/GenericHMD"), GetRootComponent());
+		HMDMeshComponent = VRMode->CreateMesh(this, AssetContainer.GenericHMDMesh, GetRootComponent());
 		HMDMeshComponent->SetMaterial(0, TeleportMID);
 		HMDMeshComponent->SetCastShadow(false);
 
@@ -162,7 +165,7 @@ void AVREditorTeleporter::StartTeleport(UViewportInteractor* Interactor)
 	TeleportingState = EState::Teleporting;
 	TeleportLerpAlpha = 0.0f;
 
-	const FVREditorAssetContainer& AssetContainer = VRMode->GetAssetContainer();
+	const UVREditorAssetContainer& AssetContainer = VRMode->GetAssetContainer();
 	VRMode->PlaySound(AssetContainer.TeleportSound, TeleportGoalLocation);
 	Interactor->PlayHapticEffect(1.0f);
 }
@@ -258,20 +261,17 @@ void AVREditorTeleporter::UpdateTeleportAim(const float DeltaTime)
 
 			if (VREd::TeleportEnableChangeScale->GetInt() != 0)
 			{
-				// Calaclate the new goal scale with the trackpad delta X axis 
-				TeleportGoalScale += VREditorInteractor->GetTrackpadSlideDelta(0) * (TeleportGoalScale * VREd::TeleportScaleSensitivity->GetFloat());
+				const float SlideDelta = GetSlideDelta(VREditorInteractor, 0);
+				if (SlideDelta != 0.0f)
+				{
+					// Calculate the new goal scale with the trackpad delta X axis 
+					TeleportGoalScale += SlideDelta * (TeleportGoalScale * VREd::TeleportScaleSensitivity->GetFloat());
 
-				const float MinScale = VRMode->GetWorldInteraction().GetMinScale() * 0.01f;
-				const float MaxScale = VRMode->GetWorldInteraction().GetMaxScale() * 0.01f;
+					const float MinScale = VRMode->GetWorldInteraction().GetMinScale() * 0.01f;
+					const float MaxScale = VRMode->GetWorldInteraction().GetMaxScale() * 0.01f;
 
-				// The scale must be in the limits of the worldinteraction minimum and maximum scale
-				if (TeleportGoalScale > MaxScale)
-				{
-					TeleportGoalScale = MaxScale;
-				}
-				else if (TeleportGoalScale < MinScale)
-				{
-					TeleportGoalScale = MinScale;
+					// The scale must be in the limits of the worldinteraction minimum and maximum scale
+					TeleportGoalScale = FMath::Clamp(TeleportGoalScale, MinScale, MaxScale);
 				}
 			}
 			else if (VREd::TeleportAllowScaleBackToDefault->GetInt() != 0)
@@ -306,7 +306,7 @@ void AVREditorTeleporter::UpdateTeleportAim(const float DeltaTime)
 			}
 
 			// The trackpad has been used while aiming for teleporting, so the teleporter won't go to the end of the laser after this
-			if (!FMath::IsNearlyZero(VREditorInteractor->GetTrackpadSlideDelta(1)) && !bPushedFromEndOfLaser && bAllowPushPull)
+			if (!bPushedFromEndOfLaser && bAllowPushPull &&  GetSlideDelta(VREditorInteractor, 1) != 0.0f)
 			{
 				bPushedFromEndOfLaser = true;
 			} 
@@ -329,7 +329,7 @@ void AVREditorTeleporter::UpdateTeleportAim(const float DeltaTime)
 
 FVector AVREditorTeleporter::UpdatePushPullTeleporter(UVREditorMotionControllerInteractor* VREditorInteractor, const FVector& LaserPointerStart, const FVector& LaserPointerEnd, const bool bEnablePushPull /* = true */)
 {
-	if (bEnablePushPull)
+	if (bEnablePushPull && GetSlideDelta(VREditorInteractor, 1) != 0.0f)
 	{
 		VREditorInteractor->CalculateDragRay(DragRayLength, DragRayLengthVelocity);
 	}
@@ -448,3 +448,26 @@ float AVREditorTeleporter::CalculateAnimatedScaleFactor() const
 	const float Scale = FMath::Max(0.001f, EasedAlpha);
 	return Scale * Scale;
 }
+
+float AVREditorTeleporter::GetSlideDelta(UVREditorMotionControllerInteractor* Interactor, const bool Axis)
+{
+	FVector2D SlideDelta = FVector2D::ZeroVector;
+	if (VRMode->GetHMDDeviceType() == EHMDDeviceType::DT_SteamVR)
+	{
+		SlideDelta = FVector2D(Interactor->GetTrackpadSlideDelta(0), Interactor->GetTrackpadSlideDelta(1));
+	}
+	else
+	{
+		SlideDelta = Interactor->GetTrackpadPosition();
+	}
+	SlideDelta.Normalize();
+
+	const bool OtherAxis = !Axis;
+	float Result = 0.0f;
+	if (!FMath::IsNearlyZero(SlideDelta[Axis], VREd::TeleportSlideBuffer->GetFloat()) && 
+		 FMath::IsNearlyZero(SlideDelta[OtherAxis], VREd::TeleportSlideBuffer->GetFloat()))
+	{
+		Result = SlideDelta.X;
+	}
+	return Result;
+}
diff --git a/Engine/Source/Editor/VREditor/Teleporter/VREditorTeleporter.h b/Engine/Source/Editor/VREditor/Teleporter/VREditorTeleporter.h
--- a/Engine/Source/Editor/VREditor/Teleporter/VREditorTeleporter.h
+++ b/Engine/Source/Editor/VREditor/Teleporter/VREditorTeleporter.h
@@ -75,7 +75,8 @@ private:
 	/** Calculated the scale for animation */
 	float CalculateAnimatedScaleFactor() const;
 
-	void ResetTeleporter();
+	/** Get slide delta to push/pull or scale the teleporter */
+	float GetSlideDelta(UVREditorMotionControllerInteractor* Interactor, const bool Axis);
 	
 	/** The owning VR mode */
 	UPROPERTY()
diff --git a/Engine/Source/Editor/VREditor/UI/VREditorDockableWindow.cpp b/Engine/Source/Editor/VREditor/UI/VREditorDockableWindow.cpp
--- a/Engine/Source/Editor/VREditor/UI/VREditorDockableWindow.cpp
+++ b/Engine/Source/Editor/VREditor/UI/VREditorDockableWindow.cpp
@@ -11,6 +11,7 @@
 #include "VREditorWidgetComponent.h"
 #include "ViewportWorldInteraction.h"
 #include "VREditorInteractor.h"
+#include "VREditorAssetContainer.h"
 
 namespace VREd
 {
@@ -33,23 +34,20 @@ AVREditorDockableWindow::AVREditorDockableWindow() :
 	CloseButtonHoverAlpha(0.0f),
 	DockSelectDistance(0.0f)
 {
+	UVREditorAssetContainer* AssetContainer = LoadObject<UVREditorAssetContainer>(nullptr, *UVREditorMode::AssetContainerPath);
 
-	UStaticMesh* WindowMesh = nullptr;
 	{
-		static ConstructorHelpers::FObjectFinder<UStaticMesh> ObjectFinder(TEXT("/Engine/VREditor/UI/SM_ContentWindow_01"));
-		WindowMesh = ObjectFinder.Object;
-		check(WindowMesh != nullptr);
+		UStaticMesh* WindowMesh = AssetContainer->WindowMesh;
+		WindowMeshComponent->SetStaticMesh(WindowMesh);
+		check(WindowMeshComponent != nullptr);
 	}
-	WindowMeshComponent->SetStaticMesh(WindowMesh);
+
+	UMaterialInterface* HoverMaterial = AssetContainer->WindowMaterial;
+	UMaterialInterface* TranslucentHoverMaterial = AssetContainer->WindowMaterial;
 
 	const FRotator RelativeRotation(30.f, 0.f, 0.f);
 	{
-		UStaticMesh* SelectionMesh = nullptr;
-		{
-			static ConstructorHelpers::FObjectFinder<UStaticMesh> ObjectFinder( TEXT( "/Engine/VREditor/TransformGizmo/SM_Dialog_Move" ) );
-			SelectionMesh = ObjectFinder.Object;
-			check( SelectionMesh != nullptr );
-		}
+		UStaticMesh* SelectionMesh = AssetContainer->WindowSelectionBarMesh;
 
 		SelectionBarMeshComponent = CreateDefaultSubobject<UStaticMeshComponent>( TEXT( "SelectionBarMesh" ) );
 		SelectionBarMeshComponent->SetStaticMesh( SelectionMesh );
@@ -63,34 +61,17 @@ AVREditorDockableWindow::AVREditorDockableWindow() :
 		SelectionBarMeshComponent->bAffectDistanceFieldLighting = false;
 		SelectionBarMeshComponent->SetRelativeRotation(RelativeRotation);
 
-		UMaterial* HoverMaterial = nullptr;
-		{
-			static ConstructorHelpers::FObjectFinder<UMaterial> ObjectFinder( TEXT( "/Engine/VREditor/TransformGizmo/TransformGizmoMaterial" ) );
-			HoverMaterial = ObjectFinder.Object;
-			check( HoverMaterial != nullptr );
-		}
+
 		SelectionBarMID = UMaterialInstanceDynamic::Create( HoverMaterial, GetTransientPackage() );
 		check( SelectionBarMID != nullptr );
 		SelectionBarMeshComponent->SetMaterial( 0, SelectionBarMID );
-
-		UMaterial* TranslucentHoverMaterial = nullptr;
-		{
-			static ConstructorHelpers::FObjectFinder<UMaterial> ObjectFinder( TEXT( "/Engine/VREditor/TransformGizmo/TranslucentTransformGizmoMaterial" ) );
-			TranslucentHoverMaterial = ObjectFinder.Object;
-			check( TranslucentHoverMaterial != nullptr );
-		}
 		SelectionBarTranslucentMID = UMaterialInstanceDynamic::Create( TranslucentHoverMaterial, GetTransientPackage() );
 		check( SelectionBarTranslucentMID != nullptr );
 		SelectionBarMeshComponent->SetMaterial( 1, SelectionBarTranslucentMID );
 	}
 
 	{
-		UStaticMesh* CloseButtonMesh = nullptr;
-		{
-			static ConstructorHelpers::FObjectFinder<UStaticMesh> ObjectFinder( TEXT( "/Engine/VREditor/TransformGizmo/SM_Dialog_Close" ) );
-			CloseButtonMesh = ObjectFinder.Object;
-			check( CloseButtonMesh != nullptr );
-		}
+		UStaticMesh* CloseButtonMesh = AssetContainer->WindowCloseButtonMesh;
 
 		CloseButtonMeshComponent= CreateDefaultSubobject<UStaticMeshComponent>( TEXT( "CloseButtonMesh" ) );
 		CloseButtonMeshComponent->SetStaticMesh( CloseButtonMesh );
@@ -104,22 +85,9 @@ AVREditorDockableWindow::AVREditorDockableWindow() :
 		CloseButtonMeshComponent->bAffectDistanceFieldLighting = false;
 		CloseButtonMeshComponent->SetRelativeRotation(RelativeRotation);
 
-		UMaterial* HoverMaterial = nullptr;
-		{
-			static ConstructorHelpers::FObjectFinder<UMaterial> ObjectFinder( TEXT( "/Engine/VREditor/TransformGizmo/TransformGizmoMaterial" ) );
-			HoverMaterial = ObjectFinder.Object;
-			check( HoverMaterial != nullptr );
-		}
 		CloseButtonMID = UMaterialInstanceDynamic::Create( HoverMaterial, GetTransientPackage() );
 		check( CloseButtonMID != nullptr );
 		CloseButtonMeshComponent->SetMaterial( 0, CloseButtonMID );
-
-		UMaterial* TranslucentHoverMaterial = nullptr;
-		{
-			static ConstructorHelpers::FObjectFinder<UMaterial> ObjectFinder( TEXT( "/Engine/VREditor/TransformGizmo/TranslucentTransformGizmoMaterial" ) );
-			TranslucentHoverMaterial = ObjectFinder.Object;
-			check( TranslucentHoverMaterial != nullptr );
-		}
 		CloseButtonTranslucentMID = UMaterialInstanceDynamic::Create( TranslucentHoverMaterial, GetTransientPackage() );
 		check( CloseButtonTranslucentMID != nullptr );
 		CloseButtonMeshComponent->SetMaterial( 1, CloseButtonTranslucentMID );
diff --git a/Engine/Source/Editor/VREditor/UI/VREditorRadialFloatingUI.cpp b/Engine/Source/Editor/VREditor/UI/VREditorRadialFloatingUI.cpp
--- a/Engine/Source/Editor/VREditor/UI/VREditorRadialFloatingUI.cpp
+++ b/Engine/Source/Editor/VREditor/UI/VREditorRadialFloatingUI.cpp
@@ -39,6 +39,8 @@ AVREditorRadialFloatingUI::AVREditorRadialFloatingUI()
 
 	DefaultGlowAmount = 2.0f;
 
+	UVREditorAssetContainer* AssetContainer = LoadObject<UVREditorAssetContainer>(nullptr, *UVREditorMode::AssetContainerPath);
+
 	{
 		WindowMeshComponent = CreateDefaultSubobject<UStaticMeshComponent>(TEXT("WindowMesh"));
 		WindowMeshComponent->SetMobility(EComponentMobility::Movable);
@@ -47,12 +49,7 @@ AVREditorRadialFloatingUI::AVREditorRadialFloatingUI()
 		WindowMeshComponent->SetCollisionEnabled(ECollisionEnabled::QueryOnly);
 		WindowMeshComponent->SetCollisionResponseToAllChannels(ECollisionResponse::ECR_Ignore);
 
-		UStaticMesh* RadialMesh = nullptr;
-		{
-			RadialMesh = LoadObject<UStaticMesh>(nullptr, TEXT("/Engine/VREditor/UI/SM_Radial_Disk"));
-			check(RadialMesh != nullptr);
-		}
-		WindowMeshComponent->SetStaticMesh(RadialMesh);
+		WindowMeshComponent->SetStaticMesh(AssetContainer->RadialMenuMainMesh);
 		WindowMeshComponent->CreateAndSetMaterialInstanceDynamic(0);
 		UMaterialInstanceDynamic* DiskMaterial = Cast<UMaterialInstanceDynamic>(WindowMeshComponent->GetMaterial(0));
 		GlowAmount = DefaultGlowAmount;
@@ -77,12 +74,7 @@ AVREditorRadialFloatingUI::AVREditorRadialFloatingUI()
 		ArrowMeshComponent->SetCollisionEnabled(ECollisionEnabled::QueryOnly);
 		ArrowMeshComponent->SetCollisionResponseToAllChannels(ECollisionResponse::ECR_Ignore);
 
-		UStaticMesh* ArrowMesh = nullptr;
-		{
-			ArrowMesh = LoadObject<UStaticMesh>(nullptr, TEXT("/Engine/VREditor/UI/SM_Radial_Pointer"));
-			check(ArrowMesh != nullptr);
-		}
-		ArrowMeshComponent->SetStaticMesh(ArrowMesh);
+		ArrowMeshComponent->SetStaticMesh(AssetContainer->RadialMenuPointerMesh);
 		ArrowMeshComponent->bGenerateOverlapEvents = false;
 		ArrowMeshComponent->SetCanEverAffectNavigation(false);
 		ArrowMeshComponent->bCastDynamicShadow = false;
@@ -388,7 +380,7 @@ void AVREditorRadialFloatingUI::ShowUI( const bool bShow, const bool bAllowFadin
 
 		if (VRMode != nullptr && bPlaySound)
 		{
-			const FVREditorAssetContainer& AssetContainer = VRMode->GetAssetContainer();
+			const UVREditorAssetContainer& AssetContainer = VRMode->GetAssetContainer();
 			VRMode->PlaySound(bShow ? AssetContainer.RadialMenuOpenSound : AssetContainer.RadialMenuCloseSound, GetActorLocation());
 		}
 
diff --git a/Engine/Source/Editor/VREditor/UI/VREditorUISystem.cpp b/Engine/Source/Editor/VREditor/UI/VREditorUISystem.cpp
--- a/Engine/Source/Editor/VREditor/UI/VREditorUISystem.cpp
+++ b/Engine/Source/Editor/VREditor/UI/VREditorUISystem.cpp
@@ -1354,7 +1354,7 @@ void UVREditorUISystem::ShowEditorUIPanel(AVREditorFloatingUI* Panel, UVREditorI
 
 		if( bPlaySound )
 		{
-			const FVREditorAssetContainer& AssetContainer = VRMode->GetAssetContainer();
+			const UVREditorAssetContainer& AssetContainer = VRMode->GetAssetContainer();
 			VRMode->PlaySound(bShouldShow ? AssetContainer.DockableWindowOpenSound : AssetContainer.DockableWindowCloseSound, Panel->GetActorLocation());
 		}
 	}
@@ -1565,7 +1565,7 @@ void UVREditorUISystem::TogglePanelsVisibility()
 	if (UIInteractor)
 	{
 		// Play sound
-		const FVREditorAssetContainer& AssetContainer = VRMode->GetAssetContainer();
+		const UVREditorAssetContainer& AssetContainer = VRMode->GetAssetContainer();
 		VRMode->PlaySound(bShowUI ? AssetContainer.DockableWindowOpenSound : AssetContainer.DockableWindowCloseSound, UIInteractor->GetTransform().GetLocation());
 	}
 }
@@ -1796,7 +1796,7 @@ TSharedRef<SWidget> UVREditorUISystem::AddHoverableButton(TSharedRef<SWidget>& B
 		Button->SetRenderTransformPivot(FVector2D(0.5, 0.5));
 		FSlateSound SlateButtonPressSound = FSlateSound();
 
-		const FVREditorAssetContainer& AssetContainer = VRMode->GetAssetContainer();
+		const UVREditorAssetContainer& AssetContainer = VRMode->GetAssetContainer();
 		SlateButtonPressSound.SetResourceObject(AssetContainer.ButtonPressSound);
 		Button->SetPressedSound(SlateButtonPressSound);
 	}
diff --git a/Engine/Source/Editor/VREditor/VREditorAssetContainer.cpp b/Engine/Source/Editor/VREditor/VREditorAssetContainer.cpp
deleted file mode 100644
diff --git a/Engine/Source/Editor/VREditor/VREditorAvatarActor.cpp b/Engine/Source/Editor/VREditor/VREditorAvatarActor.cpp
--- a/Engine/Source/Editor/VREditor/VREditorAvatarActor.cpp
+++ b/Engine/Source/Editor/VREditor/VREditorAvatarActor.cpp
@@ -19,6 +19,7 @@
 #include "UObject/Package.h"
 #include "EngineGlobals.h"
 #include "IHeadMountedDisplay.h"
+#include "VREditorAssetContainer.h"
 
 namespace VREd
 {
@@ -59,6 +60,10 @@ AVREditorAvatarActor::AVREditorAvatarActor( const FObjectInitializer& ObjectInit
 		SetRootComponent( SceneRootComponent );
 	}
 
+	// Setup the asset container.
+	UVREditorAssetContainer* AssetContainer = LoadObject<UVREditorAssetContainer>(nullptr, *UVREditorMode::AssetContainerPath);
+	check(AssetContainer != nullptr);
+
 	// Give us a head mesh
 	{
 		HeadMeshComponent = CreateDefaultSubobject<UStaticMeshComponent>( TEXT( "HeadMeshComponent" ) );
@@ -66,7 +71,7 @@ AVREditorAvatarActor::AVREditorAvatarActor( const FObjectInitializer& ObjectInit
 		HeadMeshComponent->SetupAttachment( RootComponent );
 
 		// @todo vreditor: This needs to adapt based on the device you're using
-		UStaticMesh* HeadMesh = LoadObject<UStaticMesh>( nullptr, TEXT( "/Engine/VREditor/Devices/Generic/GenericHMD" ) );
+		UStaticMesh* HeadMesh = AssetContainer->GenericHMDMesh;
 		check( HeadMesh != nullptr );
 
 		HeadMeshComponent->SetStaticMesh( HeadMesh );
@@ -81,14 +86,14 @@ AVREditorAvatarActor::AVREditorAvatarActor( const FObjectInitializer& ObjectInit
 		AddOwnedComponent( WorldMovementGridMeshComponent );
 		WorldMovementGridMeshComponent->SetupAttachment( RootComponent );
 	
-		UStaticMesh* GridMesh = LoadObject<UStaticMesh>( nullptr, TEXT( "/Engine/VREditor/WorldMovementGrid/PlaneMesh" ) );
+		UStaticMesh* GridMesh = AssetContainer->PlaneMesh;
 		check( GridMesh != nullptr );
 		WorldMovementGridMeshComponent->SetStaticMesh( GridMesh );
 		WorldMovementGridMeshComponent->SetMobility( EComponentMobility::Movable );
 		WorldMovementGridMeshComponent->SetCollisionEnabled( ECollisionEnabled::NoCollision );
 		WorldMovementGridMeshComponent->bSelectable = false;
 
-		UMaterialInterface* GridMaterial = LoadObject<UMaterialInterface>( nullptr, TEXT( "/Engine/VREditor/WorldMovementGrid/GridMaterial" ) );
+		UMaterialInterface* GridMaterial = AssetContainer->GridMaterial;
 		check( GridMaterial != nullptr );
 
 		WorldMovementGridMID = UMaterialInstanceDynamic::Create( GridMaterial, GetTransientPackage( ) );
@@ -101,10 +106,10 @@ AVREditorAvatarActor::AVREditorAvatarActor( const FObjectInitializer& ObjectInit
 
 	{
 		{
-			UMaterialInterface* UserScaleIndicatorMaterial = LoadObject<UMaterialInterface>( nullptr, TEXT( "/Engine/VREditor/LaserPointer/LaserPointerMaterialInst" ) );
+			UMaterialInterface* UserScaleIndicatorMaterial = AssetContainer->LaserPointerMaterial;
 			check( UserScaleIndicatorMaterial != nullptr );
 
-			UMaterialInstance* TranslucentUserScaleIndicatorMaterial = LoadObject<UMaterialInstance>( nullptr, TEXT( "/Engine/VREditor/LaserPointer/TranslucentLaserPointerMaterialInst" ) );
+			UMaterialInterface* TranslucentUserScaleIndicatorMaterial = AssetContainer->LaserPointerTranslucentMaterial;
 			check( TranslucentUserScaleIndicatorMaterial != nullptr );
 
 			FixedUserScaleMID = UMaterialInstanceDynamic::Create( UserScaleIndicatorMaterial, GetTransientPackage() );
@@ -119,7 +124,7 @@ AVREditorAvatarActor::AVREditorAvatarActor( const FObjectInitializer& ObjectInit
 			TranslucentCurrentUserScaleMID = UMaterialInstanceDynamic::Create( TranslucentUserScaleIndicatorMaterial, GetTransientPackage() );
 			check( TranslucentCurrentUserScaleMID != nullptr );
 
-			UStaticMesh* ScaleLineMesh = LoadObject<UStaticMesh>( nullptr, TEXT( "/Engine/VREditor/LaserPointer/LaserPointerMesh" ) );
+			UStaticMesh* ScaleLineMesh = AssetContainer->LaserPointerMesh; //@todo VREditor: The laser pointer mesh is not a closed cylinder anymore.
 			check( ScaleLineMesh != nullptr );
 
 			// Creating the background bar progress of the scale 
@@ -159,10 +164,10 @@ AVREditorAvatarActor::AVREditorAvatarActor( const FObjectInitializer& ObjectInit
 
 		// Creating the text for scaling
 		{
-			UFont* TextFont = LoadObject<UFont>( nullptr, TEXT( "/Engine/VREditor/Fonts/VRText_RobotoLarge" ) );
+			UFont* TextFont = AssetContainer->TextFont;
 			check( TextFont != nullptr );
 
-			UMaterialInterface* UserScaleIndicatorMaterial = LoadObject<UMaterialInterface>( nullptr, TEXT( "/Engine/VREditor/Fonts/VRTextMaterial" ) );
+			UMaterialInterface* UserScaleIndicatorMaterial = AssetContainer->TextMaterial;
 			check( UserScaleIndicatorMaterial != nullptr );
 
 			UserScaleIndicatorText = CreateDefaultSubobject<UTextRenderComponent>( TEXT( "UserScaleIndicatorText" ) );
@@ -194,7 +199,7 @@ AVREditorAvatarActor::AVREditorAvatarActor( const FObjectInitializer& ObjectInit
 
 	// Load our post process material for the world movement grid
 	{
-		UMaterial* Material = LoadObject<UMaterial>( nullptr, TEXT( "/Engine/VREditor/WorldMovementGrid/GridPostProcess" ) );
+		UMaterial* Material = AssetContainer->WorldMovementPostProcessMaterial;
 		check( Material != nullptr );
 		WorldMovementPostProcessMaterial = UMaterialInstanceDynamic::Create( Material, GetTransientPackage() );
 		check( WorldMovementPostProcessMaterial != nullptr );
diff --git a/Engine/Source/Editor/VREditor/VREditorMode.cpp b/Engine/Source/Editor/VREditor/VREditorMode.cpp
--- a/Engine/Source/Editor/VREditor/VREditorMode.cpp
+++ b/Engine/Source/Editor/VREditor/VREditorMode.cpp
@@ -61,6 +61,7 @@ namespace VREd
 	static FAutoConsoleVariable SFXMultiplier(TEXT("VREd.SFXMultiplier"), 1.5f, TEXT("Default Sound Effect Volume Multiplier"));
 }
 
+const FString UVREditorMode::AssetContainerPath = FString("/Engine/VREditor/VREditorAssetContainerData");
 
 UVREditorMode::UVREditorMode() : 
 	Super(),
@@ -81,7 +82,7 @@ UVREditorMode::UVREditorMode() :
 	bIsActive( false ),
 	SavedWorldToMetersScaleForPIE(0.0f),
 	bStartedPlayFromVREditor(false),
-	AssetContainer()
+	AssetContainer(nullptr)
 {
 }
 
@@ -128,12 +129,9 @@ void UVREditorMode::Init()
 		check( WorldInteraction != nullptr );
 	}
 
-	// Setup the asset editor
-	{
-		USlateWidgetStyleAsset* ContainerObject = LoadObject<USlateWidgetStyleAsset>(this, TEXT("/Engine/VREditor/VREditorAssetContainer"));
-		check(ContainerObject != nullptr);
-		AssetContainer = MakeShared<FVREditorAssetContainer>(*ContainerObject->GetStyle<FVREditorAssetContainer>());
-	}
+	// Setup the asset container.
+	AssetContainer = LoadObject<UVREditorAssetContainer>(nullptr, *UVREditorMode::AssetContainerPath);
+	check(AssetContainer != nullptr);
 
 	bIsFullyInitialized = true;
 }
@@ -150,7 +148,7 @@ void UVREditorMode::Shutdown()
 	WorldInteraction = nullptr;
 	LeftHandInteractor = nullptr;
 	RightHandInteractor = nullptr;
-	AssetContainer.Reset();
+	AssetContainer = nullptr;
 
 	// @todo vreditor urgent: Disable global editor hacks for VR Editor mode
 	GEnableVREditorHacks = false;
@@ -202,6 +200,7 @@ void UVREditorMode::Enter()
 			SavedEditorState.TransformGizmoScale = WorldInteraction->GetTransformGizmoScale();
 			WorldInteraction->SetTransformGizmoScale( VREd::GizmoScaleInVR->GetFloat() );
 			WorldInteraction->SetShouldSuppressExistingCursor(true);
+			WorldInteraction->SetInVR(true);
 		}
 	}
 
@@ -367,6 +366,7 @@ void UVREditorMode::Exit(const bool bShouldDisableStereo)
 		if( bActuallyUsingVR )
 		{
 			WorldInteraction->AddMouseCursorInteractor();
+			WorldInteraction->SetInVR(false);
 		}
 
 		WorldInteraction->SetDefaultOptionalViewportClient( nullptr );
@@ -1049,37 +1049,45 @@ void UVREditorMode::RestoreWorldToMeters()
 
 UStaticMeshComponent* UVREditorMode::CreateMotionControllerMesh( AActor* OwningActor, USceneComponent* AttachmentToComponent )
 {
-	FString MeshName;
-	if( GetHMDDeviceType() == EHMDDeviceType::DT_SteamVR )
+	UStaticMesh* ControllerMesh = nullptr;
+	if(GetHMDDeviceType() == EHMDDeviceType::DT_SteamVR)
 	{
-		MeshName = TEXT( "/Engine/VREditor/Devices/Vive/VivePreControllerMesh" );
+		ControllerMesh = AssetContainer->VivePreControllerMesh;
 	}
-	else if( GetHMDDeviceType() == EHMDDeviceType::DT_OculusRift )
+	else if(GetHMDDeviceType() == EHMDDeviceType::DT_OculusRift)
+	{
+		ControllerMesh = AssetContainer->OculusControllerMesh;
+	}
+	else
 	{
-		MeshName = TEXT( "/Engine/VREditor/Devices/Oculus/OculusControllerMesh" );
+		ControllerMesh = AssetContainer->GenericControllerMesh;
 	}
 
-	return !MeshName.IsEmpty() ? CreateMesh( OwningActor, MeshName, AttachmentToComponent ) : nullptr;
+	return CreateMesh(OwningActor, ControllerMesh, AttachmentToComponent);
 }
 
 UStaticMeshComponent* UVREditorMode::CreateMesh( AActor* OwningActor, const FString& MeshName, USceneComponent* AttachmentToComponent /*= nullptr */ )
 {
-	UStaticMeshComponent* CreatedMeshComponent = NewObject<UStaticMeshComponent>( OwningActor );
-	OwningActor->AddOwnedComponent( CreatedMeshComponent );
-	if( AttachmentToComponent != nullptr )
+	UStaticMesh* Mesh = LoadObject<UStaticMesh>(nullptr, *MeshName);
+	check(Mesh != nullptr);
+	return CreateMesh(OwningActor, Mesh, AttachmentToComponent);
+}
+
+UStaticMeshComponent* UVREditorMode::CreateMesh(AActor* OwningActor, UStaticMesh* Mesh, USceneComponent* AttachmentToComponent /*= nullptr */)
+{
+	UStaticMeshComponent* CreatedMeshComponent = NewObject<UStaticMeshComponent>(OwningActor);
+	OwningActor->AddOwnedComponent(CreatedMeshComponent);
+	if (AttachmentToComponent != nullptr)
 	{
-		CreatedMeshComponent->SetupAttachment( AttachmentToComponent );
+		CreatedMeshComponent->SetupAttachment(AttachmentToComponent);
 	}
 
 	CreatedMeshComponent->RegisterComponent();
 
-	UStaticMesh* Mesh = LoadObject<UStaticMesh>( nullptr, *MeshName );
-	check( Mesh != nullptr );
-
-	CreatedMeshComponent->SetStaticMesh( Mesh );
-	CreatedMeshComponent->SetMobility( EComponentMobility::Movable );
-	CreatedMeshComponent->SetCollisionEnabled( ECollisionEnabled::NoCollision );
-	CreatedMeshComponent->SetCollisionResponseToAllChannels( ECollisionResponse::ECR_Ignore );
+	CreatedMeshComponent->SetStaticMesh(Mesh);
+	CreatedMeshComponent->SetMobility(EComponentMobility::Movable);
+	CreatedMeshComponent->SetCollisionEnabled(ECollisionEnabled::NoCollision);
+	CreatedMeshComponent->SetCollisionResponseToAllChannels(ECollisionResponse::ECR_Ignore);
 	return CreatedMeshComponent;
 }
 
@@ -1099,7 +1107,7 @@ bool UVREditorMode::GetStartedPlayFromVREditor() const
 	return bStartedPlayFromVREditor;
 }
 
-const FVREditorAssetContainer& UVREditorMode::GetAssetContainer() const
+const UVREditorAssetContainer& UVREditorMode::GetAssetContainer() const
 {
 	return *AssetContainer;
 }
diff --git a/Engine/Source/Editor/VREditor/VREditorMotionControllerInteractor.cpp b/Engine/Source/Editor/VREditor/VREditorMotionControllerInteractor.cpp
--- a/Engine/Source/Editor/VREditor/VREditorMotionControllerInteractor.cpp
+++ b/Engine/Source/Editor/VREditor/VREditorMotionControllerInteractor.cpp
@@ -32,6 +32,7 @@
 #include "Components/SplineMeshComponent.h"
 #include "DrawDebugHelpers.h"
 #include "VREditorActions.h"
+#include "VREditorAssetContainer.h"
 
 namespace VREd
 {
@@ -212,39 +213,30 @@ void UVREditorMotionControllerInteractor::SetupComponent( AActor* OwningActor )
 		MotionControllerComponent->bDisableLowLatencyUpdate = true;
 	}
 
+	const UVREditorAssetContainer& AssetContainer = VRMode->GetAssetContainer();
+
 	// Hand mesh
 	{
 		HandMeshComponent = VRMode->CreateMotionControllerMesh( OwningActor, MotionControllerComponent ); 
 		check( HandMeshComponent != nullptr );
 		HandMeshComponent->SetCastShadow(false);
 
-		// @todo vreditor extensibility: We need this to be able to be overridden externally, or simply based on the HMD name (but allowing external folders)
-		FString MaterialName;
-		if( GetVRMode().GetHMDDeviceType() == EHMDDeviceType::DT_SteamVR )
-		{
-			MaterialName = TEXT( "/Engine/VREditor/Devices/Vive/VivePreControllerMaterial_Inst" );
-		}
-		else
-		{
-			MaterialName = TEXT( "/Engine/VREditor/Devices/Oculus/OculusControllerMaterial_Inst" );
-		}
-
-		UMaterialInstance* HandMeshMaterialInst = LoadObject<UMaterialInstance>( nullptr, *MaterialName );
-		check( HandMeshMaterialInst != nullptr );
-		HandMeshMID = UMaterialInstanceDynamic::Create( HandMeshMaterialInst, GetTransientPackage() );
+		UMaterialInterface* HandMeshMaterial = GetVRMode().GetHMDDeviceType() == EHMDDeviceType::DT_SteamVR ? AssetContainer.VivePreControllerMaterial : AssetContainer.OculusControllerMaterial;
+		check( HandMeshMaterial != nullptr );
+		HandMeshMID = UMaterialInstanceDynamic::Create( HandMeshMaterial, GetTransientPackage() );
 		check( HandMeshMID != nullptr );
 		HandMeshComponent->SetMaterial( 0, HandMeshMID );
 	}
 
 	{
-		UMaterialInstance* LaserPointerMaterialInst = LoadObject<UMaterialInstance>(nullptr, TEXT("/Engine/VREditor/LaserPointer/LaserPointerMaterialInst"));
-		check(LaserPointerMaterialInst != nullptr);
-		LaserPointerMID = UMaterialInstanceDynamic::Create(LaserPointerMaterialInst, GetTransientPackage());
+		UMaterialInterface* LaserPointerMaterial = AssetContainer.LaserPointerMaterial;
+		check(LaserPointerMaterial != nullptr);
+		LaserPointerMID = UMaterialInstanceDynamic::Create(LaserPointerMaterial, GetTransientPackage());
 		check(LaserPointerMID != nullptr);
 
-		UMaterialInstance* TranslucentLaserPointerMaterialInst = LoadObject<UMaterialInstance>(nullptr, TEXT("/Engine/VREditor/LaserPointer/TranslucentLaserPointerMaterialInst"));
-		check(TranslucentLaserPointerMaterialInst != nullptr);
-		TranslucentLaserPointerMID = UMaterialInstanceDynamic::Create(TranslucentLaserPointerMaterialInst, GetTransientPackage());
+		UMaterialInterface* TranslucentLaserPointerMaterial = AssetContainer.LaserPointerTranslucentMaterial;
+		check(TranslucentLaserPointerMaterial != nullptr);
+		TranslucentLaserPointerMID = UMaterialInstanceDynamic::Create(TranslucentLaserPointerMaterial, GetTransientPackage());
 		check(TranslucentLaserPointerMID != nullptr);
 	}
 
@@ -255,7 +247,7 @@ void UVREditorMotionControllerInteractor::SetupComponent( AActor* OwningActor )
 		HoverMeshComponent->SetupAttachment( OwningActor->GetRootComponent() );
 		HoverMeshComponent->RegisterComponent();
 
-		UStaticMesh* HoverMesh = LoadObject<UStaticMesh>( nullptr, TEXT( "/Engine/VREditor/LaserPointer/HoverMesh" ) );
+		UStaticMesh* HoverMesh = AssetContainer.LaserPointerHoverMesh;
 		check( HoverMesh != nullptr );
 		HoverMeshComponent->SetStaticMesh( HoverMesh );
 		HoverMeshComponent->SetMobility( EComponentMobility::Movable );
@@ -284,13 +276,12 @@ void UVREditorMotionControllerInteractor::SetupComponent( AActor* OwningActor )
 
 	{
 		const int32 NumLaserSplinePoints = 12;
-		UStaticMesh* MiddleSplineMesh = LoadObject<UStaticMesh>(nullptr, TEXT("/Engine/VREditor/LaserPointer/LaserPointerMesh"));
-		check (MiddleSplineMesh != nullptr);
 
-		UStaticMesh* StartSplineMesh = LoadObject<UStaticMesh>(nullptr, TEXT("/Engine/VREditor/LaserPointer/LaserPointerMesh_Start"));
+		UStaticMesh* MiddleSplineMesh = AssetContainer.LaserPointerMesh;
+		check (MiddleSplineMesh != nullptr);
+		UStaticMesh* StartSplineMesh = AssetContainer.LaserPointerStartMesh;
 		check (StartSplineMesh != nullptr);
-
-		UStaticMesh* EndSplineMesh = LoadObject<UStaticMesh>(nullptr, TEXT("/Engine/VREditor/LaserPointer/LaserPointerMesh_End"));
+		UStaticMesh* EndSplineMesh = AssetContainer.LaserPointerEndMesh;
 		check (EndSplineMesh != nullptr);
 
 		LaserSplineComponent = NewObject<USplineComponent>(OwningActor);
@@ -1443,7 +1434,7 @@ void UVREditorMotionControllerInteractor::UpdateRadialMenuInput( const float Del
 				bFlickActionExecuted == false && 
 				!IsHoveringOverUI())
 			{
-				FLevelEditorActionCallbacks::ExecuteExecCommand(FString(TEXT("TRANSACTION UNDO")));
+				VRMode->GetWorldInteraction().Undo();
 				bFlickActionExecuted = true;
 			}
 			// Move thumbstick right to redo
@@ -1451,7 +1442,7 @@ void UVREditorMotionControllerInteractor::UpdateRadialMenuInput( const float Del
 				bFlickActionExecuted == false &&
 				!IsHoveringOverUI())
 			{
-				FLevelEditorActionCallbacks::ExecuteExecCommand(FString(TEXT("TRANSACTION REDO")));
+				VRMode->GetWorldInteraction().Redo();
 				bFlickActionExecuted = true;
 			}
 			// Center to reset
@@ -1477,12 +1468,12 @@ void UVREditorMotionControllerInteractor::UndoRedoFromSwipe(const ETouchSwipeDir
 	{
 		if (InSwipeDirection == ETouchSwipeDirection::Left)
 		{
-			FLevelEditorActionCallbacks::ExecuteExecCommand(FString(TEXT("TRANSACTION UNDO")));
+			VRMode->GetWorldInteraction().Undo();
 			bFlickActionExecuted = true;
 		}
 		else if (InSwipeDirection == ETouchSwipeDirection::Right)
 		{
-			FLevelEditorActionCallbacks::ExecuteExecCommand(FString(TEXT("TRANSACTION REDO")));
+			VRMode->GetWorldInteraction().Redo();
 			bFlickActionExecuted = true;
 		}
 	}
diff --git a/Engine/Source/Editor/VREditor/VREditorWorldInteraction.cpp b/Engine/Source/Editor/VREditor/VREditorWorldInteraction.cpp
--- a/Engine/Source/Editor/VREditor/VREditorWorldInteraction.cpp
+++ b/Engine/Source/Editor/VREditor/VREditorWorldInteraction.cpp
@@ -204,7 +204,7 @@ void UVREditorWorldInteraction::OnAssetDragStartedFromContentBrowser( const TArr
 				const bool bShouldInterpolateFromDragLocation = VREd::PlacementInterpolationDuration->GetFloat() > KINDA_SMALL_NUMBER;
 				StartPlacingObjects( DroppedObjects, FactoryToUse, PlacingWithInteractor, bShouldInterpolateFromDragLocation );
 
-				const FVREditorAssetContainer& AssetContainer = Owner->GetAssetContainer();
+				const UVREditorAssetContainer& AssetContainer = Owner->GetAssetContainer();
 				Owner->PlaySound(AssetContainer.DropFromContentBrowserSound, PlacingWithInteractor->GetTransform().GetLocation());
 			}
 		}
@@ -460,7 +460,7 @@ void UVREditorWorldInteraction::PlaceDraggedMaterialOrTexture( UViewportInteract
 					const int32 TargetMaterialSlot = -1;	// All materials
 					bool bPlaced = FComponentEditorUtils::AttemptApplyMaterialToComponent( HitComponent, DroppedObjAsMaterial, TargetMaterialSlot );
 
-					const FVREditorAssetContainer& AssetContainer = Owner->GetAssetContainer();
+					const UVREditorAssetContainer& AssetContainer = Owner->GetAssetContainer();
 					Owner->PlaySound(AssetContainer.DropFromContentBrowserSound, Interactor->GetTransform().GetLocation());
 				}
 			}
diff --git a/Engine/Source/Editor/ViewportInteraction/ActorTransformer.cpp b/Engine/Source/Editor/ViewportInteraction/ActorTransformer.cpp
--- a/Engine/Source/Editor/ViewportInteraction/ActorTransformer.cpp
+++ b/Engine/Source/Editor/ViewportInteraction/ActorTransformer.cpp
@@ -26,7 +26,7 @@ void UActorTransformer::Shutdown()
 
 void UActorTransformer::OnStartDragging(class UViewportInteractor* Interactor)
 {
-	const FViewportInteractionAssetContainer& AssetContainer = ViewportWorldInteraction->GetAssetContainer();
+	const UViewportInteractionAssetContainer& AssetContainer = ViewportWorldInteraction->GetAssetContainer();
 	const FVector& SoundLocation = Interactor->GetInteractorData().GizmoLastTransform.GetLocation();
 	const EViewportInteractionDraggingMode DraggingMode = Interactor->GetDraggingMode();
 
@@ -43,7 +43,7 @@ void UActorTransformer::OnStartDragging(class UViewportInteractor* Interactor)
 
 void UActorTransformer::OnStopDragging(class UViewportInteractor* Interactor)
 {
-	const FViewportInteractionAssetContainer& AssetContainer = ViewportWorldInteraction->GetAssetContainer();
+	const UViewportInteractionAssetContainer& AssetContainer = ViewportWorldInteraction->GetAssetContainer();
 	const FVector& SoundLocation = Interactor->GetInteractorData().GizmoLastTransform.GetLocation();
 	const EViewportInteractionDraggingMode DraggingMode = Interactor->GetDraggingMode();
 
diff --git a/Engine/Source/Editor/ViewportInteraction/Gizmo/VIPivotTransformGizmo.cpp b/Engine/Source/Editor/ViewportInteraction/Gizmo/VIPivotTransformGizmo.cpp
--- a/Engine/Source/Editor/ViewportInteraction/Gizmo/VIPivotTransformGizmo.cpp
+++ b/Engine/Source/Editor/ViewportInteraction/Gizmo/VIPivotTransformGizmo.cpp
@@ -8,7 +8,8 @@
 #include "VIStretchGizmoHandle.h"
 #include "VIUniformScaleGizmoHandle.h"
 #include "ViewportWorldInteraction.h"
-#include "VREditorFloatingText.h"
+#include "ViewportInteractor.h"
+#include "ViewportInteractionAssetContainer.h"
 #include "Materials/Material.h"
 #include "Engine/Font.h"
 #include "Engine/CollisionProfile.h"
@@ -27,10 +28,13 @@ namespace VREd //@todo VREditor: Duplicates of TransformGizmo
 	static FAutoConsoleVariable PivotGizmoPlaneTranslationPivotOffsetYZ(TEXT("VI.PivotGizmoPlaneTranslationPivotOffsetYZ" ), 40.0f, TEXT( "How much the plane translation is offsetted from the pivot" ) );
 	static FAutoConsoleVariable PivotGizmoTranslationScaleMultiply( TEXT( "VI.PivotGizmoTranslationScaleMultiply" ), 2.0f, TEXT( "Multiplies translation handles scale" ) );
 	static FAutoConsoleVariable PivotGizmoTranslationHoverScaleMultiply( TEXT( "VI.PivotGizmoTranslationHoverScaleMultiply" ), 0.75f, TEXT( "Multiplies translation handles hover scale" ) );
+	static FAutoConsoleVariable PivotGizmoAimAtShrinkSize( TEXT( "VI.PivotGizmoAimAtShrinkSize" ), 0.3f, TEXT( "The minimum size when not aiming at the gizmo (0 to 1)" ) );
+	static FAutoConsoleVariable PivotGizmoAimAtAnimationSpeed( TEXT( "VI.PivotGizmoAimAtAnimationSpeed" ), 0.15f, TEXT( "The speed to animate to the gizmo full size when aiming at it" ) );
 }
 
 APivotTransformGizmo::APivotTransformGizmo() :
-	Super()
+	Super(),
+	AimingAtMeFadeAlpha(0.0f)
 {
 	UniformScaleGizmoHandleGroup = CreateDefaultSubobject<UUniformScaleGizmoHandleGroup>( TEXT( "UniformScaleHandles" ), true );
 	UniformScaleGizmoHandleGroup->SetOwningTransformGizmo(this);
@@ -89,10 +93,46 @@ void APivotTransformGizmo::UpdateGizmo(const EGizmoHandleTypes InGizmoType, cons
 
 	// Increase scale with distance, to make gizmo handles easier to click on
 	const float WorldSpaceDistanceToToPivot = FMath::Max(VREd::PivotGizmoMinDistanceForScaling->GetFloat(), FMath::Sqrt(FVector::DistSquared( GetActorLocation(), InViewLocation )));
-	const float WorldScaleFactor = GetWorld()->GetWorldSettings()->WorldToMeters / 100.0f;
-	const float GizmoScale(InScaleMultiplier * ((WorldSpaceDistanceToToPivot / WorldScaleFactor) * VREd::PivotGizmoDistanceScaleFactor->GetFloat() ) * WorldScaleFactor);
+	const float WorldScaleFactor = WorldInteraction->GetWorldScaleFactor();
+	float GizmoScale = (InScaleMultiplier * ((WorldSpaceDistanceToToPivot / WorldScaleFactor) * VREd::PivotGizmoDistanceScaleFactor->GetFloat())) * WorldScaleFactor;
 	const bool bIsWorldSpaceGizmo = (WorldInteraction->GetTransformGizmoCoordinateSpace() == COORD_World);
 
+	// Only scale the gizmo down when not aiming at it for VR implementations
+	if (WorldInteraction->IsInVR())
+	{
+		bool bIsAimingTowards = false;
+		const float GizmoRadius = (GizmoScale * 350) * 0.5f; //@todo ViewportInteraction: Hardcoded radius multiplier.
+
+		// Check if any interactor has a laser close to the gizmo
+		for (UViewportInteractor* Interactor : WorldInteraction->GetInteractors())
+		{	
+			FVector LaserPointerStart, LaserPointerEnd;
+			if (Interactor->GetLaserPointer(/*Out*/ LaserPointerStart, /*Out*/ LaserPointerEnd))
+			{
+				const FVector ClosestPointOnLaser = FMath::ClosestPointOnLine(LaserPointerStart, LaserPointerEnd, InLocalToWorld.GetLocation());
+				const float ClosestPointDistance = (ClosestPointOnLaser - InLocalToWorld.GetLocation()).Size();
+				if (ClosestPointDistance <= GizmoRadius)
+				{
+					bIsAimingTowards = true;
+					break;
+				}
+			}
+		}
+
+		const float DeltaTime = WorldInteraction->GetCurrentDeltaTime();
+		if (bIsAimingTowards)
+		{
+			AimingAtMeFadeAlpha += DeltaTime / VREd::PivotGizmoAimAtAnimationSpeed->GetFloat();
+		}
+		else
+		{
+			AimingAtMeFadeAlpha -= DeltaTime / VREd::PivotGizmoAimAtAnimationSpeed->GetFloat();
+		}
+		AimingAtMeFadeAlpha = FMath::Clamp(AimingAtMeFadeAlpha, VREd::PivotGizmoAimAtShrinkSize->GetFloat(), 1.0f);
+
+		GizmoScale *= AimingAtMeFadeAlpha;
+	}
+
 	// Update animation
 	float AnimationAlpha = GetAnimationAlpha();
 
@@ -120,14 +160,8 @@ void APivotTransformGizmo::UpdateGizmo(const EGizmoHandleTypes InGizmoType, cons
 UPivotTranslationGizmoHandleGroup::UPivotTranslationGizmoHandleGroup() :
 	Super()
 {
-	UStaticMesh* TranslationHandleMesh = nullptr;
-	{
-		static ConstructorHelpers::FObjectFinder<UStaticMesh> ObjectFinder( TEXT( "/Engine/VREditor/TransformGizmo/TranslateArrowHandle" ) );
-		TranslationHandleMesh = ObjectFinder.Object;
-		check( TranslationHandleMesh != nullptr );
-	}
-
-	CreateHandles( TranslationHandleMesh, FString( "PivotTranslationHandle" ) );
+	const UViewportInteractionAssetContainer& AssetContainer = UViewportWorldInteraction::LoadAssetContainer();
+	CreateHandles( AssetContainer.TranslationHandleMesh, FString( "PivotTranslationHandle" ) );
 }
 
 
@@ -160,14 +194,8 @@ EGizmoHandleTypes UPivotTranslationGizmoHandleGroup::GetHandleType() const
 UPivotScaleGizmoHandleGroup::UPivotScaleGizmoHandleGroup() :
 	Super()
 {
-	UStaticMesh* ScaleHandleMesh = nullptr;
-	{
-		static ConstructorHelpers::FObjectFinder<UStaticMesh> ObjectFinder( TEXT( "/Engine/VREditor/TransformGizmo/UniformScaleHandle" ) );
-		ScaleHandleMesh = ObjectFinder.Object;
-		check( ScaleHandleMesh != nullptr );
-	}
-
-	CreateHandles( ScaleHandleMesh, FString( "PivotScaleHandle" ) );	
+	const UViewportInteractionAssetContainer& AssetContainer = UViewportWorldInteraction::LoadAssetContainer();
+	CreateHandles( AssetContainer.UniformScaleHandleMesh, FString( "PivotScaleHandle" ) );	
 }
 
 void UPivotScaleGizmoHandleGroup::UpdateGizmoHandleGroup( const FTransform& LocalToWorld, const FBox& LocalBounds, const FVector ViewLocation, const bool bAllHandlesVisible, class UActorComponent* DraggingHandle, 
@@ -201,14 +229,8 @@ bool UPivotScaleGizmoHandleGroup::SupportsWorldCoordinateSpace() const
 UPivotPlaneTranslationGizmoHandleGroup::UPivotPlaneTranslationGizmoHandleGroup() :
 	Super()
 {
-	UStaticMesh* TranslationHandleMesh = nullptr;
-	{
-		static ConstructorHelpers::FObjectFinder<UStaticMesh> ObjectFinder( TEXT( "/Engine/VREditor/TransformGizmo/PlaneTranslationHandle" ) );
-		TranslationHandleMesh = ObjectFinder.Object;
-		check( TranslationHandleMesh != nullptr );
-	}
-
-	CreateHandles( TranslationHandleMesh, FString( "PlaneTranslationHandle" ) );
+	const UViewportInteractionAssetContainer& AssetContainer = UViewportWorldInteraction::LoadAssetContainer();
+	CreateHandles( AssetContainer.PlaneTranslationHandleMesh, FString( "PlaneTranslationHandle" ) );
 }
 
 void UPivotPlaneTranslationGizmoHandleGroup::UpdateGizmoHandleGroup( const FTransform& LocalToWorld, const FBox& LocalBounds, const FVector ViewLocation, const bool bAllHandlesVisible, class UActorComponent* DraggingHandle,
@@ -238,24 +260,18 @@ EGizmoHandleTypes UPivotPlaneTranslationGizmoHandleGroup::GetHandleType() const
 UPivotRotationGizmoHandleGroup::UPivotRotationGizmoHandleGroup() :
 	Super()
 {
-	UStaticMesh* QuarterRotationHandleMesh = nullptr;
-	{
-		static ConstructorHelpers::FObjectFinder<UStaticMesh> ObjectFinder(TEXT("/Engine/VREditor/TransformGizmo/RotationHandleQuarter"));
-		QuarterRotationHandleMesh = ObjectFinder.Object;
-		check(QuarterRotationHandleMesh != nullptr);
-	}
+	const UViewportInteractionAssetContainer& AssetContainer = UViewportWorldInteraction::LoadAssetContainer();
 
+	UStaticMesh* QuarterRotationHandleMesh = AssetContainer.RotationHandleMesh;
 	CreateHandles(QuarterRotationHandleMesh, FString("RotationHandle"));
 
-
 	{
 		RootFullRotationHandleComponent = CreateDefaultSubobject<USceneComponent>(TEXT("RootFullRotationHandleComponent"));
 		RootFullRotationHandleComponent->SetMobility(EComponentMobility::Movable);
 		RootFullRotationHandleComponent->SetupAttachment(this);
 
-		static ConstructorHelpers::FObjectFinder<UStaticMesh> ObjectFinder(TEXT("/Engine/VREditor/TransformGizmo/RotationHandleFull"));
-		UStaticMesh* FullRotationHandleMesh = ObjectFinder.Object;
-		check(QuarterRotationHandleMesh != nullptr);
+		UStaticMesh* FullRotationHandleMesh = AssetContainer.RotationHandleSelectedMesh;
+		check(FullRotationHandleMesh != nullptr);
 
 		FullRotationHandleMeshComponent = CreateMeshHandle(FullRotationHandleMesh, FString("FullRotationHandle"));
 		FullRotationHandleMeshComponent->SetVisibility(false);
@@ -264,8 +280,7 @@ UPivotRotationGizmoHandleGroup::UPivotRotationGizmoHandleGroup() :
 	}
 
 	{
-		static ConstructorHelpers::FObjectFinder<UStaticMesh> ObjectFinder(TEXT("/Engine/VREditor/TransformGizmo/StartRotationHandleIndicator"));
-		UStaticMesh* RotationHandleIndicatorMesh = ObjectFinder.Object;
+		UStaticMesh* RotationHandleIndicatorMesh = AssetContainer.StartRotationIndicatorMesh;
 		check(RotationHandleIndicatorMesh != nullptr);
 
 		//Start rotation indicator
@@ -275,8 +290,7 @@ UPivotRotationGizmoHandleGroup::UPivotRotationGizmoHandleGroup() :
 	}
 
 	{
-		static ConstructorHelpers::FObjectFinder<UStaticMesh> ObjectFinder(TEXT("/Engine/VREditor/TransformGizmo/RotationHandleIndicator"));
-		UStaticMesh* RotationHandleIndicatorMesh = ObjectFinder.Object;
+		UStaticMesh* RotationHandleIndicatorMesh = AssetContainer.CurrentRotationIndicatorMesh;
 		check(RotationHandleIndicatorMesh != nullptr);
 
 		//Delta rotation indicator
@@ -286,18 +300,14 @@ UPivotRotationGizmoHandleGroup::UPivotRotationGizmoHandleGroup() :
 	}
 
 	{
-		UMaterialInstance* MaterialInst = LoadObject<UMaterialInstance>(nullptr, TEXT("/Engine/VREditor/TransformGizmo/TransformGizmoMaterial_Inst"));
-		check(MaterialInst != nullptr);
-		UMaterialInstanceDynamic* DynamicMaterialInst = UMaterialInstanceDynamic::Create(MaterialInst, GetTransientPackage());
+		UMaterialInstanceDynamic* DynamicMaterialInst = UMaterialInstanceDynamic::Create(AssetContainer.TransformGizmoMaterial, GetTransientPackage());
 		check(DynamicMaterialInst != nullptr);
 
 		DeltaRotationIndicatorMeshComponent->SetMaterial(0, DynamicMaterialInst);
 		StartRotationIndicatorMeshComponent->SetMaterial(0, DynamicMaterialInst);
 		FullRotationHandleMeshComponent->SetMaterial(0, DynamicMaterialInst);
 
-		UMaterialInstance* TranslucentMaterialInst = LoadObject<UMaterialInstance>(nullptr, TEXT("/Engine/VREditor/TransformGizmo/TranslucentTransformGizmoMaterial_Inst"));
-		check(TranslucentMaterialInst != nullptr);
-		UMaterialInstanceDynamic* TranslucentDynamicMaterialInst = UMaterialInstanceDynamic::Create(TranslucentMaterialInst, GetTransientPackage());
+		UMaterialInstanceDynamic* TranslucentDynamicMaterialInst = UMaterialInstanceDynamic::Create(AssetContainer.TranslucentTransformGizmoMaterial, GetTransientPackage());
 		check(TranslucentDynamicMaterialInst != nullptr);
 
 		DeltaRotationIndicatorMeshComponent->SetMaterial(1, TranslucentDynamicMaterialInst);
diff --git a/Engine/Source/Editor/ViewportInteraction/Gizmo/VIPivotTransformGizmo.h b/Engine/Source/Editor/ViewportInteraction/Gizmo/VIPivotTransformGizmo.h
--- a/Engine/Source/Editor/ViewportInteraction/Gizmo/VIPivotTransformGizmo.h
+++ b/Engine/Source/Editor/ViewportInteraction/Gizmo/VIPivotTransformGizmo.h
@@ -52,6 +52,8 @@ private:
 	/** Stretch handle group component */
 	UPROPERTY()
 	class UStretchGizmoHandleGroup* StretchGizmoHandleGroup;
+
+	float AimingAtMeFadeAlpha;
 };
 
 /**
diff --git a/Engine/Source/Editor/ViewportInteraction/Public/ViewportInteractionAssetContainer.h b/Engine/Source/Editor/ViewportInteraction/Public/ViewportInteractionAssetContainer.h
--- a/Engine/Source/Editor/ViewportInteraction/Public/ViewportInteractionAssetContainer.h
+++ b/Engine/Source/Editor/ViewportInteraction/Public/ViewportInteractionAssetContainer.h
@@ -2,74 +2,95 @@
 
 #pragma once
 
-#include "SlateWidgetStyleContainerBase.h"
-#include "SlateWidgetStyle.h"
-#include "Sound/SoundBase.h"
+#include "Engine/DataAsset.h"
 #include "ViewportInteractionAssetContainer.generated.h"
 
+// Forward declarations
+class USoundBase;
+
 /**
- * Represents the common menu sounds used in viewport interaction
+ * Asset container for viewport interaction.
  */
-USTRUCT()
-struct FViewportInteractionAssetContainer : public FSlateWidgetStyle
+UCLASS()
+class VIEWPORTINTERACTION_API UViewportInteractionAssetContainer : public UDataAsset
 {
-	GENERATED_USTRUCT_BODY()
+	GENERATED_BODY()
 
-	FViewportInteractionAssetContainer();
-	virtual ~FViewportInteractionAssetContainer();
+public:
 
-	// FSlateWidgetStyle
-	virtual void GetResources(TArray<const FSlateBrush*>& OutBrushes) const override;
-	static const FName TypeName;
-	virtual const FName GetTypeName() const override { return TypeName; };
-	static const FViewportInteractionAssetContainer& GetDefault();
+	//
+	// Sound
+	//
 
-	/** The sound that should play when starting the game */
 	UPROPERTY(EditAnywhere, Category = Sound)
 	USoundBase* GizmoHandleSelectedSound;
-	FViewportInteractionAssetContainer& SetGizmoHandleSelectedSound(USoundBase* InSound) { GizmoHandleSelectedSound = InSound; return *this; }
 
 	UPROPERTY(EditAnywhere, Category = Sound)
 	USoundBase* GizmoHandleDropSound;
-	FViewportInteractionAssetContainer& SetGizmoHandleDropSound(USoundBase* InSound) { GizmoHandleDropSound = InSound; return *this; }
-
 
 	UPROPERTY(EditAnywhere, Category = Sound)
 	USoundBase* SelectionChangeSound;
-	FViewportInteractionAssetContainer& SetSelectionChangeSound(USoundBase* InSound) { SelectionChangeSound = InSound; return *this; }
 
 	UPROPERTY(EditAnywhere, Category = Sound)
 	USoundBase* SelectionDropSound;
-	FViewportInteractionAssetContainer& SetSelectionDropSound(USoundBase* InSound) { SelectionDropSound = InSound; return *this; }
 
 	UPROPERTY(EditAnywhere, Category = Sound)
 	USoundBase* SelectionStartDragSound;
-	FViewportInteractionAssetContainer& SetSelectionStartDragSound(USoundBase* InSound) { SelectionStartDragSound = InSound; return *this; }
 
 	UPROPERTY(EditAnywhere, Category = Sound)
 	USoundBase* GridSnapSound;
-	FViewportInteractionAssetContainer& SetGridSnapSound(USoundBase* InSound) { GridSnapSound = InSound; return *this; }
 
 	UPROPERTY(EditAnywhere, Category = Sound)
 	USoundBase* ActorSnapSound;
-	FViewportInteractionAssetContainer& SetActorSnapSound(USoundBase* InSound) { ActorSnapSound = InSound; return *this; }
-};
 
+	UPROPERTY(EditAnywhere, Category = Sound)
+	USoundBase* UndoSound;
 
-/**
- */
-UCLASS(hidecategories=Object, MinimalAPI)
-class UViewportInteractionStyleContainer : public USlateWidgetStyleContainerBase
-{
-	GENERATED_UCLASS_BODY()
+	UPROPERTY(EditAnywhere, Category = Sound)
+	USoundBase* RedoSound;
 
-public:
-	/** The actual data describing the sounds */
-	UPROPERTY(Category=Appearance, EditAnywhere, meta=(ShowOnlyInnerProperties))
-	FViewportInteractionAssetContainer SoundsStyle;
-
-	virtual const struct FSlateWidgetStyle* const GetStyle() const override
-	{
-		return static_cast< const struct FSlateWidgetStyle* >( &SoundsStyle );
-	}
+	//
+	// Meshes
+	//
+
+	UPROPERTY(EditAnywhere, Category = Mesh)
+	UStaticMesh* GridMesh;
+
+	UPROPERTY(EditAnywhere, Category = Mesh)
+	UStaticMesh* TranslationHandleMesh;
+
+	UPROPERTY(EditAnywhere, Category = Mesh)
+	UStaticMesh* UniformScaleHandleMesh;
+
+	UPROPERTY(EditAnywhere, Category = Mesh)
+	UStaticMesh* ScaleHandleMesh;
+
+	UPROPERTY(EditAnywhere, Category = Mesh)
+	UStaticMesh* PlaneTranslationHandleMesh;
+
+	UPROPERTY(EditAnywhere, Category = Mesh)
+	UStaticMesh* RotationHandleMesh;
+
+	UPROPERTY(EditAnywhere, Category = Mesh)
+	UStaticMesh* RotationHandleSelectedMesh;
+
+	UPROPERTY(EditAnywhere, Category = Mesh)
+	UStaticMesh* StartRotationIndicatorMesh;
+
+	UPROPERTY(EditAnywhere, Category = Mesh)
+	UStaticMesh* CurrentRotationIndicatorMesh;
+
+
+	//
+	// Materials
+	//
+
+	UPROPERTY(EditAnywhere, Category = Material)
+	UMaterialInterface* GridMaterial;
+
+	UPROPERTY(EditAnywhere, Category = Material)
+	UMaterialInterface* TransformGizmoMaterial;
+
+	UPROPERTY(EditAnywhere, Category = Material)
+	UMaterialInterface* TranslucentTransformGizmoMaterial;
 };
diff --git a/Engine/Source/Editor/ViewportInteraction/Public/ViewportWorldInteraction.h b/Engine/Source/Editor/ViewportInteraction/Public/ViewportWorldInteraction.h
--- a/Engine/Source/Editor/ViewportInteraction/Public/ViewportWorldInteraction.h
+++ b/Engine/Source/Editor/ViewportInteraction/Public/ViewportWorldInteraction.h
@@ -331,11 +331,20 @@ public:
 	};
 
 	/** Gets the container for all the assets of ViewportInteraction. */
-	const struct FViewportInteractionAssetContainer& GetAssetContainer() const;
+	const class UViewportInteractionAssetContainer& GetAssetContainer() const;
+
+	/** Static function to load the asset container */
+	static const class UViewportInteractionAssetContainer& LoadAssetContainer();
 
 	/** Plays sound at location. */
 	void PlaySound(USoundBase* SoundBase, const FVector& InWorldLocation, const float InVolume = 1.0f);
 
+	/** Set if this world interaction is in VR. */
+	void SetInVR(const bool bInVR);
+
+	/** Get if this world interaction is in VR. */
+	bool IsInVR() const;
+
 protected:
 
 	virtual void TransitionWorld(UWorld* NewWorld) override;
@@ -471,6 +480,9 @@ public:
 	/** Gets the color from color type */
 	FLinearColor GetColor(const EColors Color, const float Multiplier = 1.f) const;
 
+	/** The path of the asset container */
+	static const FString AssetContainerPath;
+
 private:
 
 	// All the colors for this mode
@@ -673,6 +685,12 @@ private:
 	UPROPERTY()
 	TArray<TWeakObjectPtr<AActor>> ActorsToExcludeFromHitTest;
 
+	//
+	// VR
+	//
+	
+	/** If this world interaction is in VR */
+	bool bIsInVR;
 
 	//
 	// Configuration
@@ -699,10 +717,10 @@ private:
 	/** The current tick number */
 	uint32 CurrentTickNumber;
 
-	//
-	// Assets
-	//
-
 	/** Container of assets */
-	TSharedPtr<const struct FViewportInteractionAssetContainer> AssetContainer;
+	UPROPERTY()
+	class UViewportInteractionAssetContainer* AssetContainer;
+
+	/** If we want to skip playing the sound when refreshing the transform gizmo next time */
+	bool bPlayNextRefreshTransformGizmoSound;
 };
diff --git a/Engine/Source/Editor/ViewportInteraction/ViewportInteractionAssetContainer.cpp b/Engine/Source/Editor/ViewportInteraction/ViewportInteractionAssetContainer.cpp
deleted file mode 100644
diff --git a/Engine/Source/Editor/ViewportInteraction/ViewportWorldInteraction.cpp b/Engine/Source/Editor/ViewportInteraction/ViewportWorldInteraction.cpp
--- a/Engine/Source/Editor/ViewportInteraction/ViewportWorldInteraction.cpp
+++ b/Engine/Source/Editor/ViewportInteraction/ViewportWorldInteraction.cpp
@@ -83,6 +83,8 @@ namespace VI
 	static FAutoConsoleVariable SFXMultiplier(TEXT("VI.SFXMultiplier"), 1.5f, TEXT("Default Sound Effect Volume Multiplier"));
 }
 
+const FString UViewportWorldInteraction::AssetContainerPath = FString("/Engine/VREditor/ViewportInteractionAssetContainerData");
+
 struct FGuideData
 {
 	const AActor* AlignedActor;
@@ -283,13 +285,15 @@ UViewportWorldInteraction::UViewportWorldInteraction():
 	DraggedInteractable( nullptr ),
 	DefaultMouseCursorInteractor( nullptr ),
 	DefaultMouseCursorInteractorRefCount( 0 ),
+	bIsInVR( false ),
 	bActive( false ),
 	bUseInputPreprocessor( false ),
 	bAllowWorldMovement( true ),
 	CurrentDeltaTime(0.0f),
 	bShouldSuppressCursor(false),
 	CurrentTickNumber(0),
-	AssetContainer()
+	AssetContainer(nullptr),
+	bPlayNextRefreshTransformGizmoSound(true)
 {
 }
 
@@ -308,12 +312,9 @@ void UViewportWorldInteraction::Init()
 
 	AppTimeEntered = FTimespan::FromSeconds( FApp::GetCurrentTime() );
 
-	// Setup the asset editor
-	{
-		USlateWidgetStyleAsset* ContainerObject = LoadObject<USlateWidgetStyleAsset>(this, TEXT("/Engine/VREditor/ViewportInteractionAssetContainer"));
-		check(ContainerObject != nullptr);
-		AssetContainer = MakeShared<FViewportInteractionAssetContainer>(*ContainerObject->GetStyle<FViewportInteractionAssetContainer>());
-	}
+	// Setup the asset container.
+	AssetContainer = LoadObject<UViewportInteractionAssetContainer>(nullptr, *UViewportWorldInteraction::AssetContainerPath);
+	check(AssetContainer != nullptr);
 
 	IViewportWorldInteractionManager& WorldInteractionManager = IViewportInteractionModule::Get().GetWorldInteractionManager();
 	WorldInteractionManager.SetCurrentViewportWorldInteraction( this );
@@ -390,7 +391,7 @@ void UViewportWorldInteraction::Shutdown()
 		ViewportTransformer = nullptr;
 	}
 
-	AssetContainer.Reset();
+	AssetContainer = nullptr;
 
 	USelection::SelectionChangedEvent.RemoveAll( this );
 	GEditor->OnEditorClose().RemoveAll( this );
@@ -697,11 +698,15 @@ FEditorViewportClient* UViewportWorldInteraction::GetDefaultOptionalViewportClie
 
 void UViewportWorldInteraction::Undo()
 {
+	PlaySound(AssetContainer->UndoSound, TransformGizmoActor->GetActorLocation());
+	bPlayNextRefreshTransformGizmoSound = false;
 	ExecCommand(FString("TRANSACTION UNDO"));
 }
 
 void UViewportWorldInteraction::Redo()
 {
+	PlaySound(AssetContainer->RedoSound, TransformGizmoActor->GetActorLocation());
+	bPlayNextRefreshTransformGizmoSound = false;
 	ExecCommand(FString("TRANSACTION REDO"));
 }
 
@@ -2800,10 +2805,12 @@ void UViewportWorldInteraction::RefreshTransformGizmo( const bool bNewObjectsSel
 	        SnapGridActor->GetRootComponent()->SetVisibility( bShouldBeVisible, bPropagateToChildren );
 		}
 
-		if (bNewObjectsSelected)
+		if (bNewObjectsSelected && bPlayNextRefreshTransformGizmoSound)
 		{
 			PlaySound(AssetContainer->SelectionChangeSound, TransformGizmoActor->GetActorLocation());
 		}
+
+		bPlayNextRefreshTransformGizmoSound = true;
 	}
 	else
 	{
@@ -3065,9 +3072,14 @@ bool UViewportWorldInteraction::ShouldSuppressExistingCursor() const
 
 }
 
-const FViewportInteractionAssetContainer& UViewportWorldInteraction::GetAssetContainer() const
+const UViewportInteractionAssetContainer& UViewportWorldInteraction::GetAssetContainer() const
+{
+	return *AssetContainer;
+}
+
+const class UViewportInteractionAssetContainer& UViewportWorldInteraction::LoadAssetContainer()
 {
-	return *AssetContainer.Get();
+	return *LoadObject<UViewportInteractionAssetContainer>(nullptr, *UViewportWorldInteraction::AssetContainerPath);
 }
 
 void UViewportWorldInteraction::PlaySound(USoundBase* SoundBase, const FVector& InWorldLocation, const float InVolume /*= 1.0f*/)
@@ -3079,6 +3091,16 @@ void UViewportWorldInteraction::PlaySound(USoundBase* SoundBase, const FVector&
 	}
 }
 
+void UViewportWorldInteraction::SetInVR(const bool bInVR)
+{
+	bIsInVR = bInVR;
+}
+
+bool UViewportWorldInteraction::IsInVR() const
+{
+	return bIsInVR;
+}
+
 EGizmoHandleTypes UViewportWorldInteraction::GetCurrentGizmoType() const
 {
 	switch( GLevelEditorModeTools().GetWidgetMode() )
@@ -3172,13 +3194,13 @@ void UViewportWorldInteraction::SpawnGridMeshActor()
 		SnapGridActor->SetRootComponent( SnapGridMeshComponent );
 		SnapGridMeshComponent->RegisterComponent();
 
-		UStaticMesh* GridMesh = LoadObject<UStaticMesh>( nullptr, TEXT( "/Engine/VREditor/SnapGrid/SnapGridPlaneMesh" ) );
+		UStaticMesh* GridMesh = AssetContainer->GridMesh;
 		check( GridMesh != nullptr );
 		SnapGridMeshComponent->SetStaticMesh( GridMesh );
 		SnapGridMeshComponent->SetMobility( EComponentMobility::Movable );
 		SnapGridMeshComponent->SetCollisionEnabled( ECollisionEnabled::NoCollision );
 
-		UMaterialInterface* GridMaterial = LoadObject<UMaterialInterface>( nullptr, TEXT( "/Engine/VREditor/SnapGrid/SnapGridMaterial" ) );
+		UMaterialInterface* GridMaterial = AssetContainer->GridMaterial;
 		check( GridMaterial != nullptr );
 
 		SnapGridMID = UMaterialInstanceDynamic::Create( GridMaterial, GetTransientPackage() );
diff --git a/README.md b/README.md
--- a/README.md
+++ b/README.md
@@ -1,7 +1,7 @@
 Unreal Engine
 =============
 
-Note: This branch contains all of the VR Editor and Mesh Editor features demonstrated at the **State of Unreal** keynote at the 2017 Game Developers Conference. The VR Editor and Mesh Editor are experimental features. For this particular release, features were tested exclusively on Windows, so there may be complications compiling for Mac or Linux. For documentation on the new features and updates, please see our [VR & Mesh Editor GDC 2017 Features Preview](https://docs.unrealengine.com/latest/INT/Engine/Editor/VR/GDC2017/index.html) docs.
+**Note:** This branch contains all of the VR Editor and Mesh Editor features demonstrated at the State of Unreal keynote at the 2017 Game Developers Conference. The VR Editor and Mesh Editor are experimental features. For this particular release, features were tested exclusively on Windows, so there may be complications compiling for Mac or Linux. For documentation on the new features and updates, please see our VR & Mesh Editor GDC 2017 Features Preview docs.
 
 Welcome to the Unreal Engine source code! 
 
