VWI: Fixed crash in rotation gizmo when rotating geometry
- Because the interactor data still had it's dragging mode set to 'dragging', when transformables were reset the gizmo would be updated, which would call GetDragOperationComponent()->GetDragOperation(), expecting it to be valid.

diff --git a/Engine/Source/Editor/ViewportInteraction/ViewportWorldInteraction.cpp b/Engine/Source/Editor/ViewportInteraction/ViewportWorldInteraction.cpp
--- a/Engine/Source/Editor/ViewportInteraction/ViewportWorldInteraction.cpp
+++ b/Engine/Source/Editor/ViewportInteraction/ViewportWorldInteraction.cpp
@@ -2200,6 +2200,7 @@ void UViewportWorldInteraction::StopDragging( UViewportInteractor* Interactor )
 					InteractorData.DragOperationComponent->ClearDragOperation();
 				}
 				InteractorData.DragOperationComponent.Reset();
+				InteractorData.DraggingMode = EViewportInteractionDraggingMode::Nothing;
 
 				// If we're not dragging anything around, check to see if transformables have come to a rest.
 				const bool bTransformablesStillMoving =
