Viewport Interaction: Fixed orbit controls (Alt + LMB drag) not working when VWI is active

diff --git a/Engine/Source/Editor/ViewportInteraction/ViewportWorldInteraction.cpp b/Engine/Source/Editor/ViewportInteraction/ViewportWorldInteraction.cpp
--- a/Engine/Source/Editor/ViewportInteraction/ViewportWorldInteraction.cpp
+++ b/Engine/Source/Editor/ViewportInteraction/ViewportWorldInteraction.cpp
@@ -560,19 +560,26 @@ bool UViewportWorldInteraction::InputKey( FEditorViewportClient* InViewportClien
 
 	if( IsActive() )
 	{
-		check(InViewportClient != nullptr);
-		OnKeyInputEvent.Broadcast(InViewportClient, Key, Event, bWasHandled);
+		check( InViewportClient != nullptr );
 
-		if( !bWasHandled || !IsActive() )
+		// Don't bother processing clicks if Alt is held down, as that's used for orbit in level viewports
+		FInputEventState InputState( InViewportClient->Viewport, Key, Event );
+		const bool bAltOrbitDragging = InputState.IsAltButtonPressed() && !InputState.IsCtrlButtonPressed() && !InputState.IsShiftButtonPressed() && Key == EKeys::LeftMouseButton;
+		if( !bAltOrbitDragging )
 		{
-			for( UViewportInteractor* Interactor : Interactors )
-			{
-				bWasHandled = Interactor->HandleInputKey( *InViewportClient, Key, Event );
+			OnKeyInputEvent.Broadcast(InViewportClient, Key, Event, bWasHandled);
 
-				// Stop iterating if the input was handled by an Interactor
-				if( bWasHandled || !IsActive() )
+			if( !bWasHandled || !IsActive() )
+			{
+				for( UViewportInteractor* Interactor : Interactors )
 				{
-					break;
+					bWasHandled = Interactor->HandleInputKey( *InViewportClient, Key, Event );
+
+					// Stop iterating if the input was handled by an Interactor
+					if( bWasHandled || !IsActive() )
+					{
+						break;
+					}
 				}
 			}
 		}
