Now broadcast an additional event when transforming an actor in viewport world interaction

diff --git a/Engine/Source/Editor/ViewportInteraction/ActorViewportTransformable.cpp b/Engine/Source/Editor/ViewportInteraction/ActorViewportTransformable.cpp
--- a/Engine/Source/Editor/ViewportInteraction/ActorViewportTransformable.cpp
+++ b/Engine/Source/Editor/ViewportInteraction/ActorViewportTransformable.cpp
@@ -26,10 +26,28 @@ void FActorViewportTransformable::ApplyTransform( const FTransform& NewTransform
 				}
 			}
 
-			GEditor->BroadcastBeginObjectMovement(*Actor);
 			const bool bOnlyTranslationChanged =
 				ExistingTransform.GetRotation() == NewTransform.GetRotation() &&
 				ExistingTransform.GetScale3D() == NewTransform.GetScale3D();
+			GEditor->BroadcastBeginObjectMovement(*Actor);
+			
+			UProperty* TransformProperty = FindField<UProperty>(USceneComponent::StaticClass(), GET_MEMBER_NAME_CHECKED(USceneComponent, RelativeLocation));
+			if (!bOnlyTranslationChanged)
+			{
+				if (ExistingTransform.GetRotation() != NewTransform.GetRotation())
+				{
+					TransformProperty = FindField<UProperty>(USceneComponent::StaticClass(), GET_MEMBER_NAME_CHECKED(USceneComponent, RelativeRotation));
+				}
+				else if (ExistingTransform.GetScale3D() != NewTransform.GetScale3D())
+				{
+					TransformProperty = FindField<UProperty>(USceneComponent::StaticClass(), GET_MEMBER_NAME_CHECKED(USceneComponent, RelativeScale3D));
+				}
+			}
+			
+			FEditPropertyChain PropertyChain;
+			PropertyChain.AddHead(TransformProperty);
+			FCoreUObjectDelegates::OnPreObjectPropertyChanged.Broadcast(Actor, PropertyChain);
+
 
 			Actor->SetActorTransform( NewTransform, bSweep );
 			//GWarn->Logf( TEXT( "SMOOTH: Actor %s to %s" ), *Actor->GetName(), *Transformable.TargetTransform.ToString() );
@@ -41,8 +59,13 @@ void FActorViewportTransformable::ApplyTransform( const FTransform& NewTransform
 				Actor->InvalidateLightingCacheDetailed( bOnlyTranslationChanged );
 			}
 
+
+
+			FPropertyChangedEvent PropertyChangedEvent(TransformProperty, EPropertyChangeType::ValueSet);
+			// Broadcast Post Edit change notification, we can't call PostEditChangeProperty directly on Actor or ActorComponent from here since it wasn't pair with a proper PreEditChange
+			FCoreUObjectDelegates::OnObjectPropertyChanged.Broadcast(Actor, PropertyChangedEvent);
 			const bool bFinished = false;	// @todo gizmo: PostEditChange never called; and bFinished=true never known!!
-			Actor->PostEditMove( bFinished );
+			Actor->PostEditMove(bFinished);
 			GEditor->BroadcastEndObjectMovement(*Actor);
 		}
 	}
