Fix for crash on spamming trigger presses in VR mode

diff --git a/Engine/Source/Editor/ViewportInteraction/ViewportInteractor.cpp b/Engine/Source/Editor/ViewportInteraction/ViewportInteractor.cpp
--- a/Engine/Source/Editor/ViewportInteraction/ViewportInteractor.cpp
+++ b/Engine/Source/Editor/ViewportInteraction/ViewportInteractor.cpp
@@ -174,7 +174,10 @@ bool UViewportInteractor::HandleInputKey( FEditorViewportClient& ViewportClient,
 
 							FVector LaserPointerStart, LaserPointerEnd;
 							const bool bHaveLaserPointer = GetLaserPointer( /* Out */ LaserPointerStart, /* Out */ LaserPointerEnd );
-							check( bHaveLaserPointer );
+							if (!bHaveLaserPointer)
+							{
+								return false;
+							}
 
 							bool bCanBeSelected = true;
 							if ( IViewportInteractableInterface* ActorInteractable = Cast<IViewportInteractableInterface>( Actor ) )
