Viewport World Interaction returns a pointer to a uobject instead of a reference, so it can be GC guarded for the local object.

diff --git a/Engine/Source/Editor/ViewportInteraction/Gizmo/VIPivotTransformGizmo.cpp b/Engine/Source/Editor/ViewportInteraction/Gizmo/VIPivotTransformGizmo.cpp
--- a/Engine/Source/Editor/ViewportInteraction/Gizmo/VIPivotTransformGizmo.cpp
+++ b/Engine/Source/Editor/ViewportInteraction/Gizmo/VIPivotTransformGizmo.cpp
@@ -19,6 +19,7 @@
 #include "VIGizmoHandleMeshComponent.h"
 #include "Math/UnitConversion.h"
 #include "ViewportInteractionDragOperations.h"
+#include "UObject/StrongObjectPtr.h"
 
 namespace VREd //@todo VREditor: Duplicates of TransformGizmo
 {
@@ -45,9 +46,9 @@ APivotTransformGizmo::APivotTransformGizmo() :
 	}
 	bIsEditorOnlyActor = true;
 
-	const UViewportInteractionAssetContainer& AssetContainer = UViewportWorldInteraction::LoadAssetContainer(); 
-	UMaterialInterface* GizmoMaterial = AssetContainer.TransformGizmoMaterial;
-	UMaterialInterface* TranslucentGizmoMaterial = AssetContainer.TranslucentTransformGizmoMaterial;
+	TStrongObjectPtr<const UViewportInteractionAssetContainer> AssetContainer(UViewportWorldInteraction::LoadAssetContainer()); 
+	UMaterialInterface* GizmoMaterial = AssetContainer->TransformGizmoMaterial;
+	UMaterialInterface* TranslucentGizmoMaterial = AssetContainer->TranslucentTransformGizmoMaterial;
 
 	UniformScaleGizmoHandleGroup = CreateDefaultSubobject<UUniformScaleGizmoHandleGroup>( TEXT( "UniformScaleHandles" ), true );
 	UniformScaleGizmoHandleGroup->SetOwningTransformGizmo(this);
@@ -193,8 +194,8 @@ UPivotTranslationGizmoHandleGroup::UPivotTranslationGizmoHandleGroup() :
 		return;
 	}
 
-	const UViewportInteractionAssetContainer& AssetContainer = UViewportWorldInteraction::LoadAssetContainer();
-	CreateHandles( AssetContainer.TranslationHandleMesh, FString( "PivotTranslationHandle" ) );
+	TStrongObjectPtr<const UViewportInteractionAssetContainer> AssetContainer(UViewportWorldInteraction::LoadAssetContainer());
+	CreateHandles( AssetContainer->TranslationHandleMesh, FString( "PivotTranslationHandle" ) );
 
 	DragOperationComponent->SetDragOperationClass(UTranslationDragOperation::StaticClass());
 }
@@ -229,8 +230,8 @@ UPivotScaleGizmoHandleGroup::UPivotScaleGizmoHandleGroup() :
 		return;
 	}
 
-	const UViewportInteractionAssetContainer& AssetContainer = UViewportWorldInteraction::LoadAssetContainer();
-	CreateHandles( AssetContainer.UniformScaleHandleMesh, FString( "PivotScaleHandle" ) );	
+	TStrongObjectPtr<const UViewportInteractionAssetContainer> AssetContainer(UViewportWorldInteraction::LoadAssetContainer());
+	CreateHandles( AssetContainer->UniformScaleHandleMesh, FString( "PivotScaleHandle" ) );	
 
 	DragOperationComponent->SetDragOperationClass(UScaleDragOperation::StaticClass());
 }
@@ -266,8 +267,8 @@ UPivotPlaneTranslationGizmoHandleGroup::UPivotPlaneTranslationGizmoHandleGroup()
 		return;
 	}
 
-	const UViewportInteractionAssetContainer& AssetContainer = UViewportWorldInteraction::LoadAssetContainer();
-	CreateHandles( AssetContainer.PlaneTranslationHandleMesh, FString( "PlaneTranslationHandle" ) );
+	TStrongObjectPtr<const UViewportInteractionAssetContainer> AssetContainer(UViewportWorldInteraction::LoadAssetContainer());
+	CreateHandles( AssetContainer->PlaneTranslationHandleMesh, FString( "PlaneTranslationHandle" ) );
 
 	DragOperationComponent->SetDragOperationClass(UPlaneTranslationDragOperation::StaticClass());
 }
@@ -299,9 +300,9 @@ UPivotRotationGizmoHandleGroup::UPivotRotationGizmoHandleGroup() :
 		return;
 	}
 
-	const UViewportInteractionAssetContainer& AssetContainer = UViewportWorldInteraction::LoadAssetContainer();
+	TStrongObjectPtr<const UViewportInteractionAssetContainer> AssetContainer(UViewportWorldInteraction::LoadAssetContainer());
 
-	UStaticMesh* QuarterRotationHandleMesh = AssetContainer.RotationHandleMesh;
+	UStaticMesh* QuarterRotationHandleMesh = AssetContainer->RotationHandleMesh;
 	CreateHandles(QuarterRotationHandleMesh, FString("RotationHandle"));
 
 	{
@@ -309,7 +310,7 @@ UPivotRotationGizmoHandleGroup::UPivotRotationGizmoHandleGroup() :
 		RootFullRotationHandleComponent->SetMobility(EComponentMobility::Movable);
 		RootFullRotationHandleComponent->SetupAttachment(this);
 	
-		UStaticMesh* FullRotationHandleMesh = AssetContainer.RotationHandleSelectedMesh;
+		UStaticMesh* FullRotationHandleMesh = AssetContainer->RotationHandleSelectedMesh;
 		check(FullRotationHandleMesh != nullptr);
 
 		FullRotationHandleMeshComponent = CreateMeshHandle(FullRotationHandleMesh, FString("FullRotationHandle"));
@@ -319,7 +320,7 @@ UPivotRotationGizmoHandleGroup::UPivotRotationGizmoHandleGroup() :
 	}
 
 	{
-		UStaticMesh* RotationHandleIndicatorMesh = AssetContainer.StartRotationIndicatorMesh;
+		UStaticMesh* RotationHandleIndicatorMesh = AssetContainer->StartRotationIndicatorMesh;
 		check(RotationHandleIndicatorMesh != nullptr);
 
 		//Start rotation indicator
@@ -329,7 +330,7 @@ UPivotRotationGizmoHandleGroup::UPivotRotationGizmoHandleGroup() :
 	}
 
 	{
-		UStaticMesh* RotationHandleIndicatorMesh = AssetContainer.CurrentRotationIndicatorMesh;
+		UStaticMesh* RotationHandleIndicatorMesh = AssetContainer->CurrentRotationIndicatorMesh;
 		check(RotationHandleIndicatorMesh != nullptr);
 
 		//Delta rotation indicator
@@ -339,14 +340,14 @@ UPivotRotationGizmoHandleGroup::UPivotRotationGizmoHandleGroup() :
 	}
 
 	{
-		UMaterialInstanceDynamic* DynamicMaterialInst = UMaterialInstanceDynamic::Create(AssetContainer.TransformGizmoMaterial, GetTransientPackage());
+		UMaterialInstanceDynamic* DynamicMaterialInst = UMaterialInstanceDynamic::Create(AssetContainer->TransformGizmoMaterial, GetTransientPackage());
 		check(DynamicMaterialInst != nullptr);
 
 		DeltaRotationIndicatorMeshComponent->SetMaterial(0, DynamicMaterialInst);
 		StartRotationIndicatorMeshComponent->SetMaterial(0, DynamicMaterialInst);
 		FullRotationHandleMeshComponent->SetMaterial(0, DynamicMaterialInst);
 
-		UMaterialInstanceDynamic* TranslucentDynamicMaterialInst = UMaterialInstanceDynamic::Create(AssetContainer.TranslucentTransformGizmoMaterial, GetTransientPackage());
+		UMaterialInstanceDynamic* TranslucentDynamicMaterialInst = UMaterialInstanceDynamic::Create(AssetContainer->TranslucentTransformGizmoMaterial, GetTransientPackage());
 		check(TranslucentDynamicMaterialInst != nullptr);
 
 		DeltaRotationIndicatorMeshComponent->SetMaterial(1, TranslucentDynamicMaterialInst);
diff --git a/Engine/Source/Editor/ViewportInteraction/Gizmo/VIUniformScaleGizmoHandle.cpp b/Engine/Source/Editor/ViewportInteraction/Gizmo/VIUniformScaleGizmoHandle.cpp
--- a/Engine/Source/Editor/ViewportInteraction/Gizmo/VIUniformScaleGizmoHandle.cpp
+++ b/Engine/Source/Editor/ViewportInteraction/Gizmo/VIUniformScaleGizmoHandle.cpp
@@ -9,6 +9,7 @@
 #include "ViewportWorldInteraction.h"
 #include "ViewportInteractionDragOperations.h"
 #include "ViewportInteractionAssetContainer.h"
+#include "UObject/StrongObjectPtr.h"
 
 UUniformScaleGizmoHandleGroup::UUniformScaleGizmoHandleGroup()
 	: Super(),
@@ -20,8 +21,8 @@ UUniformScaleGizmoHandleGroup::UUniformScaleGizmoHandleGroup()
 	}
 
 	// Setup uniform scaling
-	const UViewportInteractionAssetContainer& AssetContainer = UViewportWorldInteraction::LoadAssetContainer(); 
-	UStaticMesh* UniformScaleMesh = AssetContainer.UniformScaleHandleMesh;
+	TStrongObjectPtr<const UViewportInteractionAssetContainer> AssetContainer(UViewportWorldInteraction::LoadAssetContainer()); 
+	UStaticMesh* UniformScaleMesh = AssetContainer->UniformScaleHandleMesh;
 	check( UniformScaleMesh != nullptr );
 
 	UGizmoHandleMeshComponent* UniformScaleHandle = CreateMeshHandle( UniformScaleMesh, FString( "UniformScaleHandle" ) );
diff --git a/Engine/Source/Editor/ViewportInteraction/Public/ViewportWorldInteraction.h b/Engine/Source/Editor/ViewportInteraction/Public/ViewportWorldInteraction.h
--- a/Engine/Source/Editor/ViewportInteraction/Public/ViewportWorldInteraction.h
+++ b/Engine/Source/Editor/ViewportInteraction/Public/ViewportWorldInteraction.h
@@ -354,7 +354,7 @@ public:
 	const UViewportInteractionAssetContainer& GetAssetContainer() const;
 
 	/** Static function to load the asset container */
-	static const UViewportInteractionAssetContainer& LoadAssetContainer();
+	static const UViewportInteractionAssetContainer* LoadAssetContainer();
 
 	/** Plays sound at location. */
 	void PlaySound(USoundBase* SoundBase, const FVector& InWorldLocation, const float InVolume = 1.0f);
diff --git a/Engine/Source/Editor/ViewportInteraction/ViewportWorldInteraction.cpp b/Engine/Source/Editor/ViewportInteraction/ViewportWorldInteraction.cpp
--- a/Engine/Source/Editor/ViewportInteraction/ViewportWorldInteraction.cpp
+++ b/Engine/Source/Editor/ViewportInteraction/ViewportWorldInteraction.cpp
@@ -316,7 +316,7 @@ void UViewportWorldInteraction::Init()
 	AppTimeEntered = FTimespan::FromSeconds( FApp::GetCurrentTime() );
 
 	// Setup the asset container.
-	AssetContainer = &LoadAssetContainer();
+	AssetContainer = LoadAssetContainer();
 
 	// Start with the default transformer
 	SetTransformer( nullptr );
@@ -3150,11 +3150,11 @@ const UViewportInteractionAssetContainer& UViewportWorldInteraction::GetAssetCon
 	return *AssetContainer;
 }
 
-const class UViewportInteractionAssetContainer& UViewportWorldInteraction::LoadAssetContainer()
+const UViewportInteractionAssetContainer* UViewportWorldInteraction::LoadAssetContainer()
 {
 	UViewportInteractionAssetContainer* AssetContainer = LoadObject<UViewportInteractionAssetContainer>(nullptr, UViewportWorldInteraction::AssetContainerPath);
 	checkf(AssetContainer, TEXT("Failed to load ViewportInteractionAssetContainer (%s). See log for reason."), UViewportWorldInteraction::AssetContainerPath);
-	return *AssetContainer;
+	return AssetContainer;
 }
 
 void UViewportWorldInteraction::PlaySound(USoundBase* SoundBase, const FVector& InWorldLocation, const float InVolume /*= 1.0f*/)
@@ -3785,7 +3785,7 @@ void UViewportWorldInteraction::UseVWInteractions()
 	InitColors();
 
 	// Setup the asset container.
-	AssetContainer = &LoadAssetContainer();
+	AssetContainer = LoadAssetContainer();
 
 	if (DefaultMouseCursorInteractorRefCount == 0)
 	{
