#ue5 - Truncation fixes for modules:
StatsViewer
StatusBar
TargetDeviceServices
TargetPlatform
Teleporter
TimeManagement
TranslationEditor
ViewportInteraction
Virtualization
WindowsPlatformFeatures
WindowsTargetPlatform

diff --git a/Engine/Plugins/Bridge/Source/MegascansPlugin/Private/AssetImporters/DHIImport.cpp b/Engine/Plugins/Bridge/Source/MegascansPlugin/Private/AssetImporters/DHIImport.cpp
--- a/Engine/Plugins/Bridge/Source/MegascansPlugin/Private/AssetImporters/DHIImport.cpp
+++ b/Engine/Plugins/Bridge/Source/MegascansPlugin/Private/AssetImporters/DHIImport.cpp
@@ -587,7 +587,7 @@ void FImportDHI::ImportAsset(TSharedPtr<FJsonObject> AssetImportJson)
 
 		FString CharacterCopyMsg = TEXT("Re-Importing : ") + CharacterName;
 		FText CharacterCopyMsgDialogMessage = FText::FromString(CharacterCopyMsg);
-		FScopedSlowTask CharacterLoadprogress((float)SourceCharacterFiles.Num() * 3.0, CharacterCopyMsgDialogMessage);
+		FScopedSlowTask CharacterLoadprogress((float)SourceCharacterFiles.Num() * 3.0f, CharacterCopyMsgDialogMessage);
 		CharacterLoadprogress.MakeDialog();
 		
 
diff --git a/Engine/Source/Developer/TargetDeviceServices/Private/Services/TargetDeviceService.cpp b/Engine/Source/Developer/TargetDeviceServices/Private/Services/TargetDeviceService.cpp
--- a/Engine/Source/Developer/TargetDeviceServices/Private/Services/TargetDeviceService.cpp
+++ b/Engine/Source/Developer/TargetDeviceServices/Private/Services/TargetDeviceService.cpp
@@ -291,7 +291,7 @@ bool FTargetDeviceService::StoreDeployedFile(FArchive* FileReader, const FString
 
 	// copy file contents
 	int64 BytesRemaining = FileReader->TotalSize();
-	int32 BufferSize = 128 * 1024;
+	int64 BufferSize = 128 * 1024;
 
 	if (BytesRemaining < BufferSize)
 	{
diff --git a/Engine/Source/Developer/TargetPlatform/Private/TargetPlatformManagerModule.cpp b/Engine/Source/Developer/TargetPlatform/Private/TargetPlatformManagerModule.cpp
--- a/Engine/Source/Developer/TargetPlatform/Private/TargetPlatformManagerModule.cpp
+++ b/Engine/Source/Developer/TargetPlatform/Private/TargetPlatformManagerModule.cpp
@@ -766,7 +766,7 @@ protected:
 
 		// find a set of valid target platform names (the platform DataDrivenPlatformInfo.ini file was found indicates support for the platform 
 		// exists on disk, so the TP is expected to work)
-		FScopedSlowTask SlowTask(FDataDrivenPlatformInfoRegistry::GetAllPlatformInfos().Num());
+		FScopedSlowTask SlowTask((float)FDataDrivenPlatformInfoRegistry::GetAllPlatformInfos().Num());
 		for (auto Pair : FDataDrivenPlatformInfoRegistry::GetAllPlatformInfos())
 		{
 			FName PlatformName = Pair.Key;
diff --git a/Engine/Source/Developer/TargetPlatform/Public/Common/TargetPlatformBase.h b/Engine/Source/Developer/TargetPlatform/Public/Common/TargetPlatformBase.h
--- a/Engine/Source/Developer/TargetPlatform/Public/Common/TargetPlatformBase.h
+++ b/Engine/Source/Developer/TargetPlatform/Public/Common/TargetPlatformBase.h
@@ -508,7 +508,7 @@ public:
 		return TPlatformProperties::GetZlibReplacementFormat() != nullptr ? FName(TPlatformProperties::GetZlibReplacementFormat()) : NAME_Zlib;
 	}
 
-	virtual int32 GetMemoryMappingAlignment() const override
+	virtual int64 GetMemoryMappingAlignment() const override
 	{
 		return TPlatformProperties::GetMemoryMappingAlignment();
 	}
diff --git a/Engine/Source/Developer/TargetPlatform/Public/Interfaces/ITargetPlatform.h b/Engine/Source/Developer/TargetPlatform/Public/Interfaces/ITargetPlatform.h
--- a/Engine/Source/Developer/TargetPlatform/Public/Interfaces/ITargetPlatform.h
+++ b/Engine/Source/Developer/TargetPlatform/Public/Interfaces/ITargetPlatform.h
@@ -283,7 +283,7 @@ public:
 	 *
 	 * @return alignment of memory mapping.
 	 */
-	virtual int32 GetMemoryMappingAlignment() const = 0;
+	virtual int64 GetMemoryMappingAlignment() const = 0;
 
 	/**
 	 * Generates a platform specific asset manifest given an array of FAssetData.
diff --git a/Engine/Source/Developer/TextureFormat/Public/Interfaces/ITextureFormat.h b/Engine/Source/Developer/TextureFormat/Public/Interfaces/ITextureFormat.h
--- a/Engine/Source/Developer/TextureFormat/Public/Interfaces/ITextureFormat.h
+++ b/Engine/Source/Developer/TextureFormat/Public/Interfaces/ITextureFormat.h
@@ -261,7 +261,7 @@ struct FEncodedTextureDescription
 			check(InMipIndex == NumMips - InExtendedData->NumMipsInTail);
 
 			// We want the layout for the entire tail.
-			MipTextureDescription.NumMips = InExtendedData->NumMipsInTail;
+			MipTextureDescription.NumMips = IntCastChecked<uint8>(InExtendedData->NumMipsInTail);
 		}		
 		return MipTextureDescription;
 	}
diff --git a/Engine/Source/Developer/Virtualization/Private/PackageVirtualizationProcess.cpp b/Engine/Source/Developer/Virtualization/Private/PackageVirtualizationProcess.cpp
--- a/Engine/Source/Developer/Virtualization/Private/PackageVirtualizationProcess.cpp
+++ b/Engine/Source/Developer/Virtualization/Private/PackageVirtualizationProcess.cpp
@@ -246,7 +246,7 @@ void VirtualizePackages(TConstArrayView<FString> PackagePaths, TArray<FText>& Ou
 	// Build up the info in the payload provider and the final array of payload push requests
 	FWorkspaceDomainPayloadProvider PayloadProvider;
 	TArray<Virtualization::FPushRequest> PayloadsToSubmit;
-	PayloadsToSubmit.Reserve(TotalPayloadsToVirtualize);
+	PayloadsToSubmit.Reserve( IntCastChecked<int32>(TotalPayloadsToVirtualize) );
 
 	for (FPackageInfo& PackageInfo : Packages)
 	{
diff --git a/Engine/Source/Developer/Windows/WindowsTargetPlatform/Private/LocalPcTargetDevice.h b/Engine/Source/Developer/Windows/WindowsTargetPlatform/Private/LocalPcTargetDevice.h
--- a/Engine/Source/Developer/Windows/WindowsTargetPlatform/Private/LocalPcTargetDevice.h
+++ b/Engine/Source/Developer/Windows/WindowsTargetPlatform/Private/LocalPcTargetDevice.h
@@ -275,7 +275,7 @@ public:
 
 	virtual bool TerminateProcess( const int64 ProcessId ) override
 	{
-		HANDLE ProcessHandle = OpenProcess(PROCESS_TERMINATE, false, ProcessId);
+		HANDLE ProcessHandle = OpenProcess(PROCESS_TERMINATE, false, (DWORD)ProcessId );
 
 		if (ProcessHandle != NULL)
 		{
diff --git a/Engine/Source/Editor/StatsViewer/Classes/PrimitiveStats.h b/Engine/Source/Editor/StatsViewer/Classes/PrimitiveStats.h
--- a/Engine/Source/Editor/StatsViewer/Classes/PrimitiveStats.h
+++ b/Engine/Source/Editor/StatsViewer/Classes/PrimitiveStats.h
@@ -99,15 +99,15 @@ class UPrimitiveStats : public UObject
 
 	/** Minimum radius of bounding sphere of instance in map */
 	UPROPERTY(VisibleAnywhere, AssetRegistrySearchable, Category="Stats", meta=( DisplayName = "Min R", ColumnWidth = "84", ShowTotal = "true" ) )
-	float RadiusMin;
+	double RadiusMin;
 
 	/** Maximum radius of bounding sphere of instance in map */
 	UPROPERTY(VisibleAnywhere, AssetRegistrySearchable, Category="Stats", meta=( DisplayName = "Max R", ColumnWidth = "88", ShowTotal = "true" ) )
-	float RadiusMax;
+	double RadiusMax;
 
 	/** Average radius of bounding sphere of instance in map */
 	UPROPERTY(VisibleAnywhere, AssetRegistrySearchable, Category="Stats", meta=( DisplayName = "Avg R", ColumnWidth = "86", ShowTotal = "true" ) )
-	float RadiusAvg;
+	double RadiusAvg;
 
 public:
 
diff --git a/Engine/Source/Editor/StatsViewer/Private/SStatsViewer.cpp b/Engine/Source/Editor/StatsViewer/Private/SStatsViewer.cpp
--- a/Engine/Source/Editor/StatsViewer/Private/SStatsViewer.cpp
+++ b/Engine/Source/Editor/StatsViewer/Private/SStatsViewer.cpp
@@ -96,7 +96,7 @@ void SStatsViewer::Construct( const FArguments& InArgs )
 				.Padding(0.0f)
 				[
 					SNew( SComboButton )
-					.ContentPadding(3)
+					.ContentPadding(3.f)
 					.OnGetMenuContent( this, &SStatsViewer::OnGetDisplayMenuContent )
 					.ButtonContent()
 					[
@@ -111,7 +111,7 @@ void SStatsViewer::Construct( const FArguments& InArgs )
 				[
 					SNew( SButton )
 					.Visibility( this, &SStatsViewer::OnGetStatsVisibility )
-					.ContentPadding(3)
+					.ContentPadding(3.f)
 					.OnClicked( this, &SStatsViewer::OnRefreshClicked )
 					.Content()
 					[
@@ -126,7 +126,7 @@ void SStatsViewer::Construct( const FArguments& InArgs )
 				[
 					SNew( SButton )
 					.Visibility( this, &SStatsViewer::OnGetStatsVisibility )
-					.ContentPadding(3)
+					.ContentPadding(3.f)
 					.OnClicked( this, &SStatsViewer::OnExportClicked )
 					.Content()
 					[
@@ -193,7 +193,7 @@ void SStatsViewer::Construct( const FArguments& InArgs )
 				[
 					SNew( SComboButton )
 					.Visibility( this, &SStatsViewer::OnGetStatsVisibility )
-					.ContentPadding(2)
+					.ContentPadding(2.f)
 					.OnGetMenuContent( this, &SStatsViewer::OnGetFilterMenuContent )
 					.ButtonContent()
 					[
@@ -805,7 +805,7 @@ void SStatsViewer::SetDisplayedStats( TSharedRef<IStatsPage> StatsPage )
 		{
 			CustomFilterWidget = SNew( SComboButton )
 				.Visibility( this, &SStatsViewer::OnGetObjectSetsVisibility )
-				.ContentPadding(3)
+				.ContentPadding(3.f)
 				.OnGetMenuContent( this, &SStatsViewer::OnGetObjectSetMenuContent )
 				.ButtonContent()
 				[
diff --git a/Engine/Source/Editor/StatsViewer/Private/StatsPages/PrimitiveStatsPage.cpp b/Engine/Source/Editor/StatsViewer/Private/StatsPages/PrimitiveStatsPage.cpp
--- a/Engine/Source/Editor/StatsViewer/Private/StatsPages/PrimitiveStatsPage.cpp
+++ b/Engine/Source/Editor/StatsViewer/Private/StatsPages/PrimitiveStatsPage.cpp
@@ -212,7 +212,7 @@ struct PrimitiveStatsGenerator
 			int32 LightMapWidth			= 0;
 			int32 LightMapHeight		= 0;
 			InPrimitiveComponent->GetLightMapResolution( LightMapWidth, LightMapHeight );
-			int32 LMSMResolution		= FMath::Sqrt( static_cast<float>(LightMapHeight * LightMapWidth) );
+			int32 LMSMResolution		= FMath::TruncToInt32(FMath::Sqrt( static_cast<float>(LightMapHeight * LightMapWidth) ));
 			int32 LightMapData			= 0;
 			int32 LegacyShadowMapData	= 0;
 			InPrimitiveComponent->GetLightAndShadowMapMemoryUsage( LightMapData, LegacyShadowMapData );
@@ -290,7 +290,7 @@ struct PrimitiveStatsGenerator
 				NewStatsEntry->LightsLM			= LightsLMCount;
 				NewStatsEntry->LightsOther		= (float)LightsOtherCount;
 				NewStatsEntry->LightMapData		= (float)LightMapData / 1024.0f;
-				NewStatsEntry->LMSMResolution	= LMSMResolution;
+				NewStatsEntry->LMSMResolution	= (float)LMSMResolution;
 				NewStatsEntry->VertexColorMem	= (float)VertexColorMem / 1024.0f;
 				NewStatsEntry->InstVertexColorMem = (float)InstVertexColorMem / 1024.0f;
 				NewStatsEntry->UpdateNames();
@@ -520,7 +520,7 @@ void FPrimitiveStatsPage::GenerateTotals( const TArray< TWeakObjectPtr<UObject>
 	{
 		UPrimitiveStats* TotalEntry = NewObject<UPrimitiveStats>();
 
-		TotalEntry->RadiusMin = FLT_MAX;
+		TotalEntry->RadiusMin = TNumericLimits<double>::Max();
 		TotalEntry->RadiusMax = 0.0f;
 
 		// build total entry
diff --git a/Engine/Source/Editor/StatsViewer/Private/StatsPages/ShaderCookerStatsPage.cpp b/Engine/Source/Editor/StatsViewer/Private/StatsPages/ShaderCookerStatsPage.cpp
--- a/Engine/Source/Editor/StatsViewer/Private/StatsPages/ShaderCookerStatsPage.cpp
+++ b/Engine/Source/Editor/StatsViewer/Private/StatsPages/ShaderCookerStatsPage.cpp
@@ -339,7 +339,7 @@ TSharedPtr<SWidget> FShaderCookerStatsPage::GetCustomWidget(TWeakPtr<IStatsViewe
 			.HAlign(HAlign_Fill)
 			[
 				SAssignNew(PlatformComboButton, SComboButton)
-				.ContentPadding(3)
+				.ContentPadding(3.f)
 				.OnGetMenuContent(this, &FShaderCookerStatsPage::OnGetPlatformButtonMenuContent, InParentStatsViewer)
 				.ButtonContent()
 				[
diff --git a/Engine/Source/Editor/StatsViewer/Private/StatsPages/StaticMeshLightingInfoStatsPage.cpp b/Engine/Source/Editor/StatsViewer/Private/StatsPages/StaticMeshLightingInfoStatsPage.cpp
--- a/Engine/Source/Editor/StatsViewer/Private/StatsPages/StaticMeshLightingInfoStatsPage.cpp
+++ b/Engine/Source/Editor/StatsViewer/Private/StatsPages/StaticMeshLightingInfoStatsPage.cpp
@@ -280,7 +280,7 @@ TSharedPtr<SWidget> FStaticMeshLightingInfoStatsPage::GetCustomWidget(TWeakPtr<I
 		.HAlign(HAlign_Fill)
 		[
 			SAssignNew(SwapComboButton, SComboButton)
-			.ContentPadding(3)
+			.ContentPadding(3.f)
 			.OnGetMenuContent( this, &FStaticMeshLightingInfoStatsPage::OnGetSwapComboButtonMenuContent, InParentStatsViewer )
 			.ButtonContent()
 			[
@@ -295,7 +295,7 @@ TSharedPtr<SWidget> FStaticMeshLightingInfoStatsPage::GetCustomWidget(TWeakPtr<I
 		.HAlign(HAlign_Fill)
 		[
 			SAssignNew(SetToComboButton, SComboButton)
-			.ContentPadding(3)
+			.ContentPadding(3.f)
 			.OnGetMenuContent( this, &FStaticMeshLightingInfoStatsPage::OnGetSetToComboButtonMenuContent, InParentStatsViewer )
 			.ButtonContent()
 			[
diff --git a/Engine/Source/Editor/StatusBar/Private/SStatusBar.cpp b/Engine/Source/Editor/StatusBar/Private/SStatusBar.cpp
--- a/Engine/Source/Editor/StatusBar/Private/SStatusBar.cpp
+++ b/Engine/Source/Editor/StatusBar/Private/SStatusBar.cpp
@@ -174,7 +174,8 @@ class SDrawerOverlay : public SCompoundWidget
 		if (bIsResizing && this->HasMouseCapture() && !MouseEvent.GetCursorDelta().IsZero())
 		{
 			const FVector2D LocalMousePos = InitialResizeGeometry.AbsoluteToLocal(MouseEvent.GetScreenSpacePosition());
-			const float DeltaHeight = (InitialResizeGeometry.GetLocalPositionAtCoordinates(FVector2D::ZeroVector) - LocalMousePos).Y;
+			const double NarrowingPrecision = 1.0 / 16.0;
+			const float DeltaHeight = FloatCastChecked<float>((InitialResizeGeometry.GetLocalPositionAtCoordinates(FVector2D::ZeroVector) - LocalMousePos).Y, NarrowingPrecision);
 
 			TargetHeight = FMath::Clamp(InitialHeightAtResize + DeltaHeight, MinHeight, MaxHeight);
 			SetHeight(InitialHeightAtResize + DeltaHeight);
@@ -376,7 +377,7 @@ public:
 				//.AutoWidth()
 				[
 					SNew(SBox)
-					.HeightOverride(8)
+					.HeightOverride(8.f)
 					[
 						SNew(SOverlay)
 						+ SOverlay::Slot()
@@ -1020,7 +1021,7 @@ TSharedRef<SWidget> SStatusBar::OnGetProgressBarMenuContent()
 
 	return
 		SNew(SBox)
-		.WidthOverride(ProgressBar->GetDesiredSize().X-8.0f)
+		.WidthOverride((float)ProgressBar->GetDesiredSize().X-8.0f)
 		[
 			ProgressBarMenu.MakeWidget()
 		];
diff --git a/Engine/Source/Editor/TranslationEditor/Private/TranslationEditor.cpp b/Engine/Source/Editor/TranslationEditor/Private/TranslationEditor.cpp
--- a/Engine/Source/Editor/TranslationEditor/Private/TranslationEditor.cpp
+++ b/Engine/Source/Editor/TranslationEditor/Private/TranslationEditor.cpp
@@ -817,7 +817,7 @@ void FTranslationEditor::ChangeSourceFont()
 
 	if ( bOpened && NewFontFilename != "")
 	{
-		SourceFont = FSlateFontInfo(NewFontFilename, SourceFont.Size);
+		SourceFont = FSlateFontInfo(NewFontFilename, IntCastChecked<uint16>(SourceFont.Size));
 		RefreshUI();
 	}
 }
@@ -832,7 +832,7 @@ void FTranslationEditor::ChangeTranslationTargetFont()
 
 	if ( bOpened && NewFontFilename != "")
 	{
-		TranslationTargetFont = FSlateFontInfo(NewFontFilename, TranslationTargetFont.Size);
+		TranslationTargetFont = FSlateFontInfo(NewFontFilename, IntCastChecked<uint16>(TranslationTargetFont.Size));
 		RefreshUI();
 	}
 }
diff --git a/Engine/Source/Editor/TranslationEditor/Private/TranslationPickerEditWindow.cpp b/Engine/Source/Editor/TranslationEditor/Private/TranslationPickerEditWindow.cpp
--- a/Engine/Source/Editor/TranslationEditor/Private/TranslationPickerEditWindow.cpp
+++ b/Engine/Source/Editor/TranslationEditor/Private/TranslationPickerEditWindow.cpp
@@ -75,7 +75,6 @@ void STranslationPickerEditWindow::Construct(const FArguments& InArgs)
 	for (FText PickedText : PickedTexts)
 	{
 		TSharedPtr<SEditableTextBox> TextBox;
-		int32 DefaultPadding = 0.0f;
 
 		TSharedRef<STranslationPickerEditWidget> NewEditWidget = 
 			SNew(STranslationPickerEditWidget)
@@ -96,7 +95,7 @@ void STranslationPickerEditWindow::Construct(const FArguments& InArgs)
 	}
 
 	TSharedPtr<SEditableTextBox> TextBox;
-	int32 DefaultPadding = 0.0f;
+	float DefaultPadding = 0.0f;
 
 	// Layout the Translation Picker Edit Widgets and some save/close buttons below them
 	WindowContents->SetContent(
@@ -251,7 +250,6 @@ void STranslationPickerEditWidget::Construct(const FArguments& InArgs)
 {
 	PickedText = InArgs._PickedText;
 	bAllowEditing = InArgs._bAllowEditing;
-	int32 DefaultPadding = 0.0f;
 
 	// Try and get the localization information for this text
 	FTextId TextId;
diff --git a/Engine/Source/Editor/TranslationEditor/Private/TranslationPickerFloatingWindow.cpp b/Engine/Source/Editor/TranslationEditor/Private/TranslationPickerFloatingWindow.cpp
--- a/Engine/Source/Editor/TranslationEditor/Private/TranslationPickerFloatingWindow.cpp
+++ b/Engine/Source/Editor/TranslationEditor/Private/TranslationPickerFloatingWindow.cpp
@@ -209,7 +209,7 @@ void STranslationPickerFloatingWindow::Tick( const FGeometry& AllottedGeometry,
 		DesiredPosition.X -= WindowSize.X;
 
 		// Clamp to work area
-		DesiredPosition = FSlateApplication::Get().CalculateTooltipWindowPosition(FSlateRect(DesiredPosition.X, DesiredPosition.Y, DesiredPosition.X, DesiredPosition.Y), WindowSize, false);
+		DesiredPosition = FSlateApplication::Get().CalculateTooltipWindowPosition(FSlateRect(DesiredPosition, DesiredPosition), WindowSize, false);
 
 		// also kind of a hack, but this is the only way at the moment to get a 'cursor decorator' without using the drag-drop code path
 		ParentWindow.Pin()->MoveWindowTo(DesiredPosition);
diff --git a/Engine/Source/Editor/UnrealEd/Private/EditorEngine.cpp b/Engine/Source/Editor/UnrealEd/Private/EditorEngine.cpp
--- a/Engine/Source/Editor/UnrealEd/Private/EditorEngine.cpp
+++ b/Engine/Source/Editor/UnrealEd/Private/EditorEngine.cpp
@@ -1167,7 +1167,7 @@ void UEditorEngine::Init(IEngineLoop* InEngineLoop)
 			TEXT("AnimationSettings")
 		};
 
-		FScopedSlowTask ModuleSlowTask(UE_ARRAY_COUNT(ModuleNames));
+		FScopedSlowTask ModuleSlowTask((float)UE_ARRAY_COUNT(ModuleNames));
 		for (const TCHAR* ModuleName : ModuleNames)
 		{
 			ModuleSlowTask.EnterProgressFrame(1);
diff --git a/Engine/Source/Editor/VREditor/Public/VREditorInteractor.h b/Engine/Source/Editor/VREditor/Public/VREditorInteractor.h
--- a/Engine/Source/Editor/VREditor/Public/VREditorInteractor.h
+++ b/Engine/Source/Editor/VREditor/Public/VREditorInteractor.h
@@ -80,7 +80,7 @@ public:
 	// ViewportInteractorInterface overrides
 	virtual void Shutdown_Implementation() override;
 	virtual void Tick_Implementation( const float DeltaTime ) override;
-	virtual void CalculateDragRay( float& InOutDragRayLength, float& InOutDragRayVelocity ) override;
+	virtual void CalculateDragRay( double& InOutDragRayLength, double& InOutDragRayVelocity ) override;
 
 	/** @return Returns the type of HMD we're dealing with */
 	UFUNCTION( BlueprintCallable, Category = "UVREditorInteractor" )
diff --git a/Engine/Source/Editor/VREditor/Teleporter/VREditorTeleporter.h b/Engine/Source/Editor/VREditor/Teleporter/VREditorTeleporter.h
--- a/Engine/Source/Editor/VREditor/Teleporter/VREditorTeleporter.h
+++ b/Engine/Source/Editor/VREditor/Teleporter/VREditorTeleporter.h
@@ -172,10 +172,10 @@ private:
 	float TeleportGoalScale;
 
 	/** The current length of the laser where the teleport should be at */
-	float DragRayLength;
+	double DragRayLength;
 
 	/** The current drag velocity to push or pull the teleport along the laser */
-	float DragRayLengthVelocity;
+	double DragRayLengthVelocity;
 
 	/** If the teleporter has been pushed by the trackpad */
 	bool bPushedFromEndOfLaser;
diff --git a/Engine/Source/Editor/VREditor/VREditorInteractor.cpp b/Engine/Source/Editor/VREditor/VREditorInteractor.cpp
--- a/Engine/Source/Editor/VREditor/VREditorInteractor.cpp
+++ b/Engine/Source/Editor/VREditor/VREditorInteractor.cpp
@@ -792,7 +792,7 @@ FName UVREditorInteractor::GetHMDDeviceType() const
 	return GetVRMode().GetHMDDeviceType();
 }
 
-void UVREditorInteractor::CalculateDragRay( float& InOutDragRayLength, float& InOutDragRayVelocity )
+void UVREditorInteractor::CalculateDragRay( double& InOutDragRayLength, double& InOutDragRayVelocity )
 {
 	const FTimespan CurrentTime = FTimespan::FromSeconds( FPlatformTime::Seconds() );
 	const float WorldScaleFactor = WorldInteraction->GetWorldScaleFactor();
diff --git a/Engine/Source/Editor/ViewportInteraction/Gizmo/VIBaseTransformGizmo.cpp b/Engine/Source/Editor/ViewportInteraction/Gizmo/VIBaseTransformGizmo.cpp
--- a/Engine/Source/Editor/ViewportInteraction/Gizmo/VIBaseTransformGizmo.cpp
+++ b/Engine/Source/Editor/ViewportInteraction/Gizmo/VIBaseTransformGizmo.cpp
@@ -66,11 +66,11 @@ float ABaseTransformGizmo::GetAnimationAlpha()
 	float AnimationAlpha = 0.0f;
 	{
 		const FTimespan CurrentTime = FTimespan::FromSeconds( FApp::GetCurrentTime() );
-		const float TimeSinceSelectionChange = (CurrentTime - SelectedAtTime).GetTotalSeconds();
-		const float AnimLength = VREd::GizmoSelectionAnimationDuration->GetFloat();
+		const double TimeSinceSelectionChange = (CurrentTime - SelectedAtTime).GetTotalSeconds();
+		const double AnimLength = VREd::GizmoSelectionAnimationDuration->GetFloat();
 		if ( TimeSinceSelectionChange < AnimLength )
 		{
-			AnimationAlpha = FMath::Max( 0.0f, TimeSinceSelectionChange / AnimLength );
+			AnimationAlpha = FMath::Max( 0.0f, float(TimeSinceSelectionChange / AnimLength) );
 		}
 		else
 		{
diff --git a/Engine/Source/Editor/ViewportInteraction/Gizmo/VIGizmoHandleMeshComponent.cpp b/Engine/Source/Editor/ViewportInteraction/Gizmo/VIGizmoHandleMeshComponent.cpp
--- a/Engine/Source/Editor/ViewportInteraction/Gizmo/VIGizmoHandleMeshComponent.cpp
+++ b/Engine/Source/Editor/ViewportInteraction/Gizmo/VIGizmoHandleMeshComponent.cpp
@@ -69,7 +69,7 @@ public:
 	}
 
 	virtual uint32 GetMemoryFootprint(void) const override { return sizeof *this + GetAllocatedSize(); }
-	uint32 GetAllocatedSize(void) const { return FPrimitiveSceneProxy::GetAllocatedSize(); }
+	SIZE_T GetAllocatedSize(void) const { return FPrimitiveSceneProxy::GetAllocatedSize(); }
 };
 
 UGizmoHandleMeshComponent::UGizmoHandleMeshComponent() :
diff --git a/Engine/Source/Editor/ViewportInteraction/Gizmo/VIPivotTransformGizmo.cpp b/Engine/Source/Editor/ViewportInteraction/Gizmo/VIPivotTransformGizmo.cpp
--- a/Engine/Source/Editor/ViewportInteraction/Gizmo/VIPivotTransformGizmo.cpp
+++ b/Engine/Source/Editor/ViewportInteraction/Gizmo/VIPivotTransformGizmo.cpp
@@ -106,9 +106,9 @@ void APivotTransformGizmo::UpdateGizmo(const EGizmoHandleTypes InGizmoType, cons
 	SetActorTransform(InLocalToWorld);
 
 	// Increase scale with distance, to make gizmo handles easier to click on
-	const float WorldSpaceDistanceToToPivot = FMath::Max(VREd::PivotGizmoMinDistanceForScaling->GetFloat(), FMath::Sqrt(FVector::DistSquared( GetActorLocation(), InViewLocation )));
-	const float WorldScaleFactor = WorldInteraction->GetWorldScaleFactor();
-	const float GizmoScale = (InScaleMultiplier * ((WorldSpaceDistanceToToPivot / WorldScaleFactor) * VREd::PivotGizmoDistanceScaleFactor->GetFloat())) * WorldScaleFactor;
+	const double WorldSpaceDistanceToToPivot = FMath::Max(VREd::PivotGizmoMinDistanceForScaling->GetFloat(), FMath::Sqrt(FVector::DistSquared( GetActorLocation(), InViewLocation )));
+	const double WorldScaleFactor = WorldInteraction->GetWorldScaleFactor();
+	const double GizmoScale = (InScaleMultiplier * ((WorldSpaceDistanceToToPivot / WorldScaleFactor) * VREd::PivotGizmoDistanceScaleFactor->GetFloat())) * WorldScaleFactor;
 	const bool bIsWorldSpaceGizmo = (WorldInteraction->GetTransformGizmoCoordinateSpace() == COORD_World);
 
 	if (LastDraggingHandle != nullptr && DraggingHandle == nullptr)
@@ -116,12 +116,12 @@ void APivotTransformGizmo::UpdateGizmo(const EGizmoHandleTypes InGizmoType, cons
 		AimingAtGizmoScaleAlpha = 0.0f;
 	}
 
-	float AnimatedGizmoScale = GizmoScale;
+	double AnimatedGizmoScale = GizmoScale;
 	// Only scale the gizmo down when not aiming at it for VR implementations
 	if (WorldInteraction->IsInVR())
 	{
 		bool bIsAimingTowards = false;
-		const float GizmoRadius = (GizmoScale * 350) * 0.5f;
+		const double GizmoRadius = (GizmoScale * 350) * 0.5f;
 
 		// Check if any interactor has a laser close to the gizmo
 		for (UViewportInteractor* Interactor : WorldInteraction->GetInteractors())
@@ -133,7 +133,7 @@ void APivotTransformGizmo::UpdateGizmo(const EGizmoHandleTypes InGizmoType, cons
 				if (Interactor->GetLaserPointer(/*Out*/ LaserPointerStart, /*Out*/ LaserPointerEnd))
 				{
 					const FVector ClosestPointOnLaser = FMath::ClosestPointOnLine(LaserPointerStart, LaserPointerEnd, InLocalToWorld.GetLocation());
-					const float ClosestPointDistance = (ClosestPointOnLaser - InLocalToWorld.GetLocation()).Size();
+					const double ClosestPointDistance = (ClosestPointOnLaser - InLocalToWorld.GetLocation()).Size();
 					if (ClosestPointDistance <= GizmoRadius)
 					{
 						bIsAimingTowards = true;
@@ -167,10 +167,11 @@ void APivotTransformGizmo::UpdateGizmo(const EGizmoHandleTypes InGizmoType, cons
 		{
 			bool bIsHoveringOrDraggingThisHandleGroup = false;
 			
-			const float Scale = HandleGroup == StretchGizmoHandleGroup ? GizmoScale : AnimatedGizmoScale;
+			const double Scale = HandleGroup == StretchGizmoHandleGroup ? GizmoScale : AnimatedGizmoScale;
+			const float FloatScale = FloatCastChecked<float>(Scale, 1.0 / 16.0);
 
 			HandleGroup->UpdateGizmoHandleGroup(InLocalToWorld, InLocalBounds, InViewLocation, bInAllHandlesVisible, DraggingHandle,
-				InHoveringOverHandles, AnimationAlpha, Scale, InGizmoHoverScale, InGizmoHoverAnimationDuration, /* Out */ bIsHoveringOrDraggingThisHandleGroup);
+				InHoveringOverHandles, AnimationAlpha, FloatScale, InGizmoHoverScale, InGizmoHoverAnimationDuration, /* Out */ bIsHoveringOrDraggingThisHandleGroup);
 
 			
 			if (HandleGroup != RotationGizmoHandleGroup)
@@ -493,7 +494,7 @@ EGizmoHandleTypes UPivotRotationGizmoHandleGroup::GetHandleType() const
 
 void UPivotRotationGizmoHandleGroup::UpdateIndicator(USceneComponent* IndicatorRoot, const FVector& Direction, const uint32 FacingAxisIndex)
 {
-	float Y = 0, X = 0;
+	double Y = 0, X = 0;
 	if (FacingAxisIndex == 0)
 	{
 		Y = Direction.Y;
@@ -510,7 +511,7 @@ void UPivotRotationGizmoHandleGroup::UpdateIndicator(USceneComponent* IndicatorR
 		X = -Direction.X;
 	}
 
-	const float Angle = FMath::RadiansToDegrees(FMath::Atan2(Y, X));
+	const double Angle = FMath::RadiansToDegrees(FMath::Atan2(Y, X));
 	IndicatorRoot->SetRelativeRotation(FRotator(0, 0, Angle));
 }
 
diff --git a/Engine/Source/Editor/ViewportInteraction/Public/ViewportInteractionUtils.h b/Engine/Source/Editor/ViewportInteraction/Public/ViewportInteractionUtils.h
--- a/Engine/Source/Editor/ViewportInteraction/Public/ViewportInteractionUtils.h
+++ b/Engine/Source/Editor/ViewportInteraction/Public/ViewportInteractionUtils.h
@@ -64,7 +64,7 @@ namespace ViewportInteractionUtils
 
 		const FVector CalculateCutoff(const FVector& InValue);
 		const FVector CalculateAlpha(const FVector& InCutoff, const double InDeltaTime) const;
-		const float CalculateAlpha(const float InCutoff, const double InDeltaTime) const;
+		const double CalculateAlpha(const double InCutoff, const double InDeltaTime) const;
 
 		double MinCutoff;
 		double CutoffSlope;
diff --git a/Engine/Source/Editor/ViewportInteraction/Public/ViewportInteractor.h b/Engine/Source/Editor/ViewportInteraction/Public/ViewportInteractor.h
--- a/Engine/Source/Editor/ViewportInteraction/Public/ViewportInteractor.h
+++ b/Engine/Source/Editor/ViewportInteraction/Public/ViewportInteractor.h
@@ -192,6 +192,9 @@ public:
 	virtual void ResetHoverState();
 
 	/** Needs to be implemented by the base to calculate drag ray length and the velocity for the ray */
+	virtual void CalculateDragRay( double& InOutDragRayLength, double& InOutDragRayVelocity ) {};
+
+	UE_DEPRECATED(5.2, "CalculateDragRay() now uses double-precision arguments.")
 	virtual void CalculateDragRay( float& InOutDragRayLength, float& InOutDragRayVelocity ) {};
 
 	/**
diff --git a/Engine/Source/Editor/ViewportInteraction/Public/ViewportWorldInteraction.h b/Engine/Source/Editor/ViewportInteraction/Public/ViewportWorldInteraction.h
--- a/Engine/Source/Editor/ViewportInteraction/Public/ViewportWorldInteraction.h
+++ b/Engine/Source/Editor/ViewportInteraction/Public/ViewportWorldInteraction.h
@@ -161,7 +161,7 @@ public:
 	bool PreprocessedInputKey( const FKey Key, const EInputEvent Event );
 
 	/** Handles the axis input from the Input Preprocessor*/
-	bool PreprocessedInputAxis( const int32 ControllerId, const FKey Key, const float Delta, const float DeltaTime );
+	bool PreprocessedInputAxis( const int32 ControllerId, const FKey Key, const float Delta, const double DeltaTime );
 
 	/** Gets the world space transform of the calibrated VR room origin.  When using a seated VR device, this will feel like the
 	camera's world transform (before any HMD positional or rotation adjustments are applied.) */
diff --git a/Engine/Source/Editor/ViewportInteraction/ViewportInteractionDragOperations.cpp b/Engine/Source/Editor/ViewportInteraction/ViewportInteractionDragOperations.cpp
--- a/Engine/Source/Editor/ViewportInteraction/ViewportInteractionDragOperations.cpp
+++ b/Engine/Source/Editor/ViewportInteraction/ViewportInteractionDragOperations.cpp
@@ -100,7 +100,7 @@ void URotateOnAngleDragOperation::ExecuteDrag(FDraggingTransformableData& Draggi
 			}
 
 			// Get the angle between the center and the intersected point
-			float AngleToIntersectedLocation = FMath::Atan2(RotatedIntersectLocationOnPlane.Y, RotatedIntersectLocationOnPlane.Z);
+			float AngleToIntersectedLocation = (float)FMath::Atan2(RotatedIntersectLocationOnPlane.Y, RotatedIntersectLocationOnPlane.Z);
 			if (!StartDragAngleOnRotation.IsSet())
 			{
 				StartDragAngleOnRotation = AngleToIntersectedLocation;
@@ -134,7 +134,7 @@ void UScaleDragOperation::ExecuteDrag(struct FDraggingTransformableData& Draggin
 	HandlePlacement.GetCenterHandleCountAndFacingAxisIndex(/* Out */ CenterHandleCount, /* Out */ FacingAxisIndex, /* Out */ CenterAxisIndex);
 
 	const FVector PassGizmoSpaceDraggedTo = DraggingData.GizmoStartTransform.InverseTransformPositionNoScale(DraggingData.PassDraggedTo);
-	float AddedScaleOnAxis = PassGizmoSpaceDraggedTo[FacingAxisIndex] * VI::ScaleSensitivity->GetFloat();
+	double AddedScaleOnAxis = PassGizmoSpaceDraggedTo[FacingAxisIndex] * VI::ScaleSensitivity->GetFloat();
 
 	// Invert if we we are scaling on the negative side of the gizmo
 	USceneComponent* DraggingTransformGizmoComponent = DraggingData.Interactor->GetInteractorData().DraggingTransformGizmoComponent.Get();
diff --git a/Engine/Source/Editor/ViewportInteraction/ViewportInteractionUtils.cpp b/Engine/Source/Editor/ViewportInteraction/ViewportInteractionUtils.cpp
--- a/Engine/Source/Editor/ViewportInteraction/ViewportInteractionUtils.cpp
+++ b/Engine/Source/Editor/ViewportInteraction/ViewportInteractionUtils.cpp
@@ -106,9 +106,9 @@ namespace ViewportInteractionUtils
 		return Result;
 	}
 
-	const float FOneEuroFilter::CalculateAlpha(const float InCutoff, const double InDeltaTime) const
+	const double FOneEuroFilter::CalculateAlpha(const double InCutoff, const double InDeltaTime) const
 	{
-		const double tau = 1.0 / (2 * PI * InCutoff);
+		const double tau = 1.0 / (2 * UE_DOUBLE_PI * InCutoff);
 		return 1.0 / (1.0 + tau / InDeltaTime);
 	}
 
diff --git a/Engine/Source/Editor/ViewportInteraction/ViewportInteractorData.h b/Engine/Source/Editor/ViewportInteraction/ViewportInteractorData.h
--- a/Engine/Source/Editor/ViewportInteraction/ViewportInteractorData.h
+++ b/Engine/Source/Editor/ViewportInteraction/ViewportInteractorData.h
@@ -72,7 +72,7 @@ struct FViewportInteractorData
 	bool bWasAssistingDrag;
 
 	/** Length of the ray that's dragging */
-	float DragRayLength;
+	double DragRayLength;
 
 	/** Location that we dragged to last frame (end point of the ray) */
 	FVector LastDragToLocation;
@@ -90,7 +90,7 @@ struct FViewportInteractorData
 	FVector DragTranslationVelocity;
 
 	/** How fast to adjust ray length every frame for inertial ray length changes */
-	float DragRayLengthVelocity;
+	double DragRayLengthVelocity;
 
 	/** While dragging, true if we're dragging at least one simulated object that we're driving the velocities of.  When this
 	    is true, our default inertia system is disabled and we rely on the physics engine to take care of inertia */
diff --git a/Engine/Source/Editor/ViewportInteraction/ViewportWorldInteraction.cpp b/Engine/Source/Editor/ViewportInteraction/ViewportWorldInteraction.cpp
--- a/Engine/Source/Editor/ViewportInteraction/ViewportWorldInteraction.cpp
+++ b/Engine/Source/Editor/ViewportInteraction/ViewportWorldInteraction.cpp
@@ -105,7 +105,7 @@ struct FGuideData
 	FVector GuideEnd;
 	FColor GuideColor;
 	float DrawAlpha;
-	float GuideLength;
+	double GuideLength;
 };
 
 // @todo vreditor: Hacky inline implementation of a double vector.  Move elsewhere and polish or get rid.
@@ -657,11 +657,11 @@ bool UViewportWorldInteraction::InputAxis( FEditorViewportClient* InViewportClie
 	return bWasHandled;
 }
 
-bool UViewportWorldInteraction::PreprocessedInputAxis( const int32 ControllerId, const FKey Key, const float Delta, const float DeltaTime )
+bool UViewportWorldInteraction::PreprocessedInputAxis( const int32 ControllerId, const FKey Key, const float Delta, const double DeltaTime )
 {
 	if( DefaultOptionalViewportClient != nullptr && bUseInputPreprocessor == true )
 	{
-		return InputAxis( DefaultOptionalViewportClient, DefaultOptionalViewportClient->Viewport, ControllerId, Key, Delta, DeltaTime );
+		return InputAxis( DefaultOptionalViewportClient, DefaultOptionalViewportClient->Viewport, ControllerId, Key, Delta, (float)DeltaTime );
 	}
 
 	return false;
@@ -710,7 +710,7 @@ void UViewportWorldInteraction::SetHeadTransform( const FTransform& NewHeadTrans
 		FRotator UseRotation( NewHeadTransform.GetRotation().Rotator() );
 
 		// adjust the yaw by the difference of the flag and the HMD
-		float YawDifference = GetHeadTransform().GetRotation().Rotator().Yaw - UseRotation.Yaw;
+		double YawDifference = GetHeadTransform().GetRotation().Rotator().Yaw - UseRotation.Yaw;
 		UseRotation.Yaw = GetRoomTransform().GetRotation().Rotator().Yaw - YawDifference;
 		UseRotation.Pitch = 0.0f;
 		UseRotation.Roll = 0.0f;
@@ -1084,7 +1084,7 @@ void UViewportWorldInteraction::InteractionTick( const float DeltaTime )
 							if( !bIsInterpolatingTransformablesFromSnapshotTransform )	// Let the last animation finish first
 							{
 								const FVector WorldSpaceDragDelta = InteractorData.bIsFirstDragUpdate ? FVector::ZeroVector : ( DraggedTo - InteractorData.LastDragToLocation );
-								const float ScaledDragDistance = WorldSpaceDragDelta.Size() * WorldScaleFactor;
+								const double ScaledDragDistance = WorldSpaceDragDelta.Size() * WorldScaleFactor;
 								if( ScaledDragDistance >= VI::DragAtLaserImpactInterpolationThreshold->GetFloat() )
 								{
 									bIsInterpolatingTransformablesFromSnapshotTransform = true;
@@ -1292,8 +1292,8 @@ void UViewportWorldInteraction::InteractionTick( const float DeltaTime )
 				bool bCrossedGridLine = false;
 				for( int32 AxisIndex = 0; AxisIndex < 3; ++AxisIndex )
 				{
-					const int32 OldGridCellIndex = FMath::TruncToInt( OldViewLocation[ AxisIndex ] / WorldSpaceHapticTranslationInterval );
-					const int32 NewGridCellIndex = FMath::TruncToInt( NewViewLocation[ AxisIndex ] / WorldSpaceHapticTranslationInterval );
+					const int32 OldGridCellIndex = FMath::TruncToInt32( OldViewLocation[ AxisIndex ] / WorldSpaceHapticTranslationInterval );
+					const int32 NewGridCellIndex = FMath::TruncToInt32( NewViewLocation[ AxisIndex ] / WorldSpaceHapticTranslationInterval );
 					if( OldGridCellIndex != NewGridCellIndex )
 					{
 						bCrossedGridLine = true;
@@ -1811,18 +1811,18 @@ void UViewportWorldInteraction::UpdateDragging(
 
 				const FVector LineStart = DraggedTo;
 				const FVector LineStartDelta = DragDelta;
-				const float LineStartDistance = LineStartDelta.Size();
+				const double LineStartDistance = LineStartDelta.Size();
 				const FVector LastLineStart = LineStart - LineStartDelta;
 
 				const FVector LineEnd = OtherHandDraggedTo;
 				const FVector LineEndDelta = OtherHandDragDelta;
-				const float LineEndDistance = LineEndDelta.Size();
+				const double LineEndDistance = LineEndDelta.Size();
 				const FVector LastLineEnd = LineEnd - LineEndDelta;
 
 				// Choose a point along the new line segment to rotate and scale relative to.  We'll weight it toward the
 				// side of the line that moved the most this update.
-				const float TotalDistance = LineStartDistance + LineEndDistance;
-				float LineStartToEndActivityWeight = 0.5f;	// Default to right in the center, if no distance moved yet.
+				const double TotalDistance = LineStartDistance + LineEndDistance;
+				double LineStartToEndActivityWeight = 0.5f;	// Default to right in the center, if no distance moved yet.
 				if (GetDefault<UVISettings>()->bScaleWorldWithDynamicPivot && !FMath::IsNearlyZero( TotalDistance ) )	// Avoid division by zero
 				{
 					LineStartToEndActivityWeight = LineStartDistance / TotalDistance;
@@ -1846,20 +1846,20 @@ void UViewportWorldInteraction::UpdateDragging(
 					DrawDebugSphere( GetWorld(), LinesRelativeTo.TransformPosition( PivotLocation ), 2.5f * GetWorldScaleFactor(), 32, FColor::White, false, 0.0f );
 				}
 
-				const float LastLineLength = ( LastLineEnd - LastLineStart ).Size();
-				const float LineLength = ( LineEnd - LineStart ).Size();
+				const double LastLineLength = ( LastLineEnd - LastLineStart ).Size();
+				const double LineLength = ( LineEnd - LineStart ).Size();
 				ScaleOffset = FVector( LineLength / LastLineLength );
 				//			ScaleOffset = FVector( 0.98f + 0.04f * FMath::MakePulsatingValue( FPlatformTime::Seconds(), 0.1f ) ); // FVector( LineLength / LastLineLength );
-				GizmoScaleSinceDragStarted += ( LineLength - LastLineLength ) / GetWorldScaleFactor();
+				GizmoScaleSinceDragStarted += FloatCastChecked<float>((LineLength - LastLineLength) / GetWorldScaleFactor(), 1.0 / 16.0);
 
 				// How much did the line rotate since last time?
 				RotationOffset = FQuat::FindBetweenVectors( LastLineEnd - LastLineStart, LineEnd - LineStart );
-				GizmoRotationRadiansSinceDragStarted += RotationOffset.AngularDistance( FQuat::Identity );
+				GizmoRotationRadiansSinceDragStarted += (float)RotationOffset.AngularDistance( FQuat::Identity );
 
 				// For translation, only move proportionally to the common vector between the two deltas.  Basically,
 				// you need to move both hands in the same direction to translate while gripping with two hands.
 				const FVector AverageDelta = ( DragDelta + OtherHandDragDelta ) * 0.5f;
-				const float TranslationWeight = FMath::Max( 0.0f, FVector::DotProduct( DragDelta.GetSafeNormal(), OtherHandDragDelta.GetSafeNormal() ) );
+				const double TranslationWeight = FMath::Max( 0.0f, FVector::DotProduct( DragDelta.GetSafeNormal(), OtherHandDragDelta.GetSafeNormal() ) );
 				TranslationOffset = FMath::Lerp( FVector::ZeroVector, AverageDelta, TranslationWeight );
 			}
 			else
@@ -1960,7 +1960,7 @@ void UViewportWorldInteraction::UpdateDragging(
 				if (bAllowWorldScaling)
 				{
 					// Adjust world scale
-					const float WorldScaleOffset = ScaleOffset.GetAbsMax();
+					const float WorldScaleOffset = (float)ScaleOffset.GetAbsMax();
 					if (WorldScaleOffset != 0.0f)
 					{
 						const float OldWorldToMetersScale = GetWorld()->GetWorldSettings()->WorldToMeters;
@@ -2569,9 +2569,9 @@ bool UViewportWorldInteraction::FindPlacementPointUnderLaser( UViewportInteracto
 					ExtremeDirection.X >= 0.0f ? Box.Max.X : Box.Min.X,
 					ExtremeDirection.Y >= 0.0f ? Box.Max.Y : Box.Min.Y,
 					ExtremeDirection.Z >= 0.0f ? Box.Max.Z : Box.Min.Z );
-				const float ProjectionDistance = FVector::DotProduct( ExtremePoint, ExtremeDirection );
+				const double ProjectionDistance = FVector::DotProduct( ExtremePoint, ExtremeDirection );
 
-				const float ExtraOffset = 
+				const double ExtraOffset =
 					( bAnyPhysicallySimulatedTransformables && GEditor->bIsSimulatingInEditor ) ? ( GizmoLocalBounds.GetSize().GetAbsMax() * VI::PlacementOffsetScaleWhileSimulating->GetFloat() ) : 0.0f;
 
 				ExtremePointOnBox = ExtremeDirection * ( ProjectionDistance + ExtraOffset );
@@ -2773,11 +2773,11 @@ void UViewportWorldInteraction::RefreshTransformGizmo( const bool bNewObjectsSel
 			const FBox GizmoWorldBoundsFlattened = GizmoLocalBoundsFlattened.TransformBy( GizmoToWorld );
 
 			// Make sure we're at least as big as the gizmo bounds, but also large enough that you can see at least a few grid cells (depending on your snap size)
-			const float SnapGridSize = 
+			const double SnapGridSize = 
 				FMath::Max( GizmoWorldBoundsFlattened.GetSize().GetAbsMax(), GEditor->GetGridSize() * 2.5f ) * VI::SnapGridSize->GetFloat();
 
 			// The mesh is 100x100, so we'll scale appropriately
-			const float SnapGridScale = SnapGridSize / 100.0f;
+			const double SnapGridScale = SnapGridSize / 100.0f;
 			SnapGridActor->SetActorScale3D( FVector( SnapGridScale ) );
 
 			EViewportInteractionDraggingMode DraggingMode = EViewportInteractionDraggingMode::Nothing;
@@ -2841,7 +2841,7 @@ void UViewportWorldInteraction::RefreshTransformGizmo( const bool bNewObjectsSel
 				LocalSnapGridMID->SetScalarParameterValue( GridIntervalParameterName, GridInterval );
 
 				static FName GridRadiusParameterName( "GridRadius" );
-				LocalSnapGridMID->SetScalarParameterValue( GridRadiusParameterName, SnapGridSize * 0.5f );
+				LocalSnapGridMID->SetScalarParameterValue( GridRadiusParameterName, FloatCastChecked<float>(SnapGridSize * 0.5f, 1./16.) );
 
 				static FName LineWidthParameterName( "LineWidth" );
 				LocalSnapGridMID->SetScalarParameterValue( LineWidthParameterName, LineWidth );
@@ -2955,16 +2955,16 @@ float UViewportWorldInteraction::GetTransformGizmoScale() const
 
 void UViewportWorldInteraction::ApplyVelocityDamping( FVector& Velocity, const bool bVelocitySensitive )
 {
-	const float InertialMovementZeroEpsilon = 0.01f;	// @todo vreditor tweak
+	const double InertialMovementZeroEpsilon = 0.01f;	// @todo vreditor tweak
 	if ( !Velocity.IsNearlyZero( InertialMovementZeroEpsilon ) )
 	{
 		// Apply damping
 		if ( bVelocitySensitive )
 		{
-			const float DampenMultiplierAtLowSpeeds = LowSpeedInertiaDamping;
-			const float DampenMultiplierAtHighSpeeds = HighSpeedInertiaDamping;
-			const float SpeedForMinimalDamping = 2.5f * GetWorldScaleFactor();	// cm/frame	// @todo vreditor tweak
-			const float SpeedBasedDampeningScalar = FMath::Clamp( Velocity.Size(), 0.0f, SpeedForMinimalDamping ) / SpeedForMinimalDamping;	// @todo vreditor: Probably needs a curve applied to this to compensate for our framerate insensitivity
+			const double DampenMultiplierAtLowSpeeds = LowSpeedInertiaDamping;
+			const double DampenMultiplierAtHighSpeeds = HighSpeedInertiaDamping;
+			const double SpeedForMinimalDamping = 2.5f * GetWorldScaleFactor();	// cm/frame	// @todo vreditor tweak
+			const double SpeedBasedDampeningScalar = FMath::Clamp( Velocity.Size(), 0.0f, SpeedForMinimalDamping ) / SpeedForMinimalDamping;	// @todo vreditor: Probably needs a curve applied to this to compensate for our framerate insensitivity
 			Velocity = Velocity * FMath::Lerp( DampenMultiplierAtLowSpeeds, DampenMultiplierAtHighSpeeds, SpeedBasedDampeningScalar );	// @todo vreditor: Frame rate sensitive damping.  Make use of delta time!
 		}
 		else
@@ -3409,7 +3409,7 @@ FVector UViewportWorldInteraction::FindTransformGizmoAlignPoint(const FTransform
 	}
 
 	// Our snap distances are some percentage of the transform gizmo's dimensions
-	const float AdjustedSnapDistance = (VI::ForceSnapDistance->GetFloat() / 100.0f) * 2.0f * (MinGuideLength.GetAbsMax());
+	const double AdjustedSnapDistance = (VI::ForceSnapDistance->GetFloat() / 100.0f) * 2.0f * (MinGuideLength.GetAbsMax());
 	int32 NumberOfMatchesNeeded = 0;
 	if (bShouldConstrainMovement)
 	{
@@ -3451,7 +3451,7 @@ FVector UViewportWorldInteraction::FindTransformGizmoAlignPoint(const FTransform
 		// Find all possible candidates for alignment
 		// TODO: add the world grid
 		// TODO: remove anything it might not make sense to align to
-		const float CompareDistance = VI::AlignCandidateDistance->GetFloat() * MinGuideLength.GetAbsMax();
+		const float CompareDistance = FloatCastChecked<float>(VI::AlignCandidateDistance->GetFloat() * MinGuideLength.GetAbsMax(), 1./16.);
 		const FVector Start = DesiredGizmoTransform.GetLocation();
 		const FVector End = DesiredGizmoTransform.GetLocation();
 		TArray<FOverlapResult> OutOverlaps;
@@ -3543,10 +3543,10 @@ FVector UViewportWorldInteraction::FindTransformGizmoAlignPoint(const FTransform
 	// Now find the best guide from all available point combinations
 	FGuideData AlignedGuide;
 	bool bFoundAlignedTransform = false;
-	float AlignedDeltaSize = 10000000.0f;
+	double AlignedDeltaSize = 10000000.0f;
 	for (FGuideData PotentialGizmoGuide : PotentialGizmoGuides)
 	{	
-		const float OffsetSize = PotentialGizmoGuide.LocalOffset.Size();
+		const double OffsetSize = PotentialGizmoGuide.LocalOffset.Size();
 
 		// Keep finding the guide with the shortest offset size
 		if (OffsetSize < AlignedDeltaSize && !FMath::IsNearlyZero(OffsetSize))
@@ -3594,13 +3594,13 @@ void UViewportWorldInteraction::DrawBoxBrackets(const FBox InActor, const FTrans
 
 
 			// MinVector
-			OutVectorMin.X = FMath::Min<float>(LocalBox.Min.X, OutVectorMin.X);
-			OutVectorMin.Y = FMath::Min<float>(LocalBox.Min.Y, OutVectorMin.Y);
-			OutVectorMin.Z = FMath::Min<float>(LocalBox.Min.Z, OutVectorMin.Z);
+			OutVectorMin.X = FMath::Min<double>(LocalBox.Min.X, OutVectorMin.X);
+			OutVectorMin.Y = FMath::Min<double>(LocalBox.Min.Y, OutVectorMin.Y);
+			OutVectorMin.Z = FMath::Min<double>(LocalBox.Min.Z, OutVectorMin.Z);
 			// MaxVector
-			OutVectorMax.X = FMath::Max<float>(LocalBox.Max.X, OutVectorMax.X);
-			OutVectorMax.Y = FMath::Max<float>(LocalBox.Max.Y, OutVectorMax.Y);
-			OutVectorMax.Z = FMath::Max<float>(LocalBox.Max.Z, OutVectorMax.Z);
+			OutVectorMax.X = FMath::Max<double>(LocalBox.Max.X, OutVectorMax.X);
+			OutVectorMax.Y = FMath::Max<double>(LocalBox.Max.Y, OutVectorMax.Y);
+			OutVectorMax.Z = FMath::Max<double>(LocalBox.Max.Z, OutVectorMax.Z);
 		}
 	};
 
@@ -3611,7 +3611,7 @@ void UViewportWorldInteraction::DrawBoxBrackets(const FBox InActor, const FTrans
 	Local::GetBoundingVectors(InActor, MinVector, MaxVector);
 
 	// Create a bracket offset to determine the length of our corner axises
-	const float BracketOffset = FVector::Dist(MinVector, MaxVector) * 0.1f;
+	const double BracketOffset = FVector::Dist(MinVector, MaxVector) * 0.1f;
 
 	// Calculate bracket corners based on min/max vectors
 	TArray<FVector> BracketCorners;
@@ -3645,9 +3645,9 @@ void UViewportWorldInteraction::DrawBoxBrackets(const FBox InActor, const FTrans
 		const FVector WorldBracketY = LocalToWorld.TransformPosition(LocalBracketY);
 		const FVector WorldBracketZ = LocalToWorld.TransformPosition(LocalBracketZ);
 
-		DrawDebugLine(GetWorld(), WorldCorner, WorldBracketX, GROUP_COLOR, false, -1.0f, 1.0f, 2.0f);
-		DrawDebugLine(GetWorld(), WorldCorner, WorldBracketY, GROUP_COLOR, false, -1.0f, 1.0f, 2.0f);
-		DrawDebugLine(GetWorld(), WorldCorner, WorldBracketZ, GROUP_COLOR, false, -1.0f, 1.0f, 2.0f);
+		DrawDebugLine(GetWorld(), WorldCorner, WorldBracketX, GROUP_COLOR, false, -1.0f, 1, 2.0f);
+		DrawDebugLine(GetWorld(), WorldCorner, WorldBracketY, GROUP_COLOR, false, -1.0f, 1, 2.0f);
+		DrawDebugLine(GetWorld(), WorldCorner, WorldBracketZ, GROUP_COLOR, false, -1.0f, 1, 2.0f);
 	}
 
 }
diff --git a/Engine/Source/Runtime/TimeManagement/Private/FixedFrameRateCustomTimeStep.cpp b/Engine/Source/Runtime/TimeManagement/Private/FixedFrameRateCustomTimeStep.cpp
--- a/Engine/Source/Runtime/TimeManagement/Private/FixedFrameRateCustomTimeStep.cpp
+++ b/Engine/Source/Runtime/TimeManagement/Private/FixedFrameRateCustomTimeStep.cpp
@@ -24,8 +24,8 @@ void UFixedFrameRateCustomTimeStep::WaitForFixedFrameRate() const
 	const FFrameRate FrameRate = GetFixedFrameRate();
 
 	// Calculate delta time
-	const float DeltaRealTime = CurrentTime - FApp::GetLastTime();
-	const float WaitTime = FMath::Max(FrameRate.AsInterval() - DeltaRealTime, 0.0);
+	const double DeltaRealTime = CurrentTime - FApp::GetLastTime();
+	const double WaitTime = FMath::Max(FrameRate.AsInterval() - DeltaRealTime, 0.0);
 
 	double ActualWaitTime = 0.0;
 	{
@@ -33,7 +33,7 @@ void UFixedFrameRateCustomTimeStep::WaitForFixedFrameRate() const
 
 		if (WaitTime > 5.f / 1000.f)
 		{
-			FPlatformProcess::SleepNoStats(WaitTime - 0.002f);
+			FPlatformProcess::SleepNoStats((float)WaitTime - 0.002f);
 		}
 
 		// Give up timeslice for remainder of wait time.
diff --git a/Engine/Source/Runtime/TimeManagement/Private/GenlockedCustomTimeStep.cpp b/Engine/Source/Runtime/TimeManagement/Private/GenlockedCustomTimeStep.cpp
--- a/Engine/Source/Runtime/TimeManagement/Private/GenlockedCustomTimeStep.cpp
+++ b/Engine/Source/Runtime/TimeManagement/Private/GenlockedCustomTimeStep.cpp
@@ -39,6 +39,6 @@ uint32 UGenlockedCustomTimeStep::GetExpectedSyncCountDelta() const
 	//Depending on format, sync count per frame will not be 1.
 	//For example, PsF will have 2 sync counts
 	const FFrameRate DeltaFrameRate = GetSyncRate() / GetFixedFrameRate();
-	return FMath::RoundToInt(DeltaFrameRate.AsDecimal());
+	return FMath::RoundToInt32(DeltaFrameRate.AsDecimal());
 }
 
diff --git a/Engine/Source/Runtime/TimeManagement/Private/GenlockedFixedRateCustomTimeStep.cpp b/Engine/Source/Runtime/TimeManagement/Private/GenlockedFixedRateCustomTimeStep.cpp
--- a/Engine/Source/Runtime/TimeManagement/Private/GenlockedFixedRateCustomTimeStep.cpp
+++ b/Engine/Source/Runtime/TimeManagement/Private/GenlockedFixedRateCustomTimeStep.cpp
@@ -80,7 +80,7 @@ bool UGenlockedFixedRateCustomTimeStep::WaitForSync()
 			// Normal sleep for the bulk of the idle time.
 			if (LastIdleTime > 5.f / 1000.f)
 			{
-				FPlatformProcess::SleepNoStats(LastIdleTime - 0.002f);
+				FPlatformProcess::SleepNoStats((float)LastIdleTime - 0.002f);
 			}
 
 			// Give up timeslice for small remainder of wait time.
diff --git a/Engine/Source/Runtime/TimeManagement/Private/GenlockedTimecodeProvider.cpp b/Engine/Source/Runtime/TimeManagement/Private/GenlockedTimecodeProvider.cpp
--- a/Engine/Source/Runtime/TimeManagement/Private/GenlockedTimecodeProvider.cpp
+++ b/Engine/Source/Runtime/TimeManagement/Private/GenlockedTimecodeProvider.cpp
@@ -125,7 +125,7 @@ void UGenlockedTimecodeProvider::FetchAndUpdate()
 
 				if (CalcRoundedRateRatio(GenlockRate, LastFrameTime.Rate, RateRatio))
 				{
-					LastFrameTime.Rate.Numerator *= RateRatio; // RateRatio should already be rounded.
+					LastFrameTime.Rate.Numerator = FMath::TruncToInt32(LastFrameTime.Rate.Numerator * RateRatio); // RateRatio should already be rounded.
 					LastFrameTime.Time = FFrameTime::FromDecimal(LastFrameTime.Time.AsDecimal() * RateRatio);
 				}
 			}
diff --git a/Engine/Source/Runtime/TimeManagement/Private/TimeManagementTests.cpp b/Engine/Source/Runtime/TimeManagement/Private/TimeManagementTests.cpp
--- a/Engine/Source/Runtime/TimeManagement/Private/TimeManagementTests.cpp
+++ b/Engine/Source/Runtime/TimeManagement/Private/TimeManagementTests.cpp
@@ -103,8 +103,8 @@ bool FFrameTimeToSecondsTest::RunTest(const FString& Parameters)
 		FFrameTime Time     = TestTimes[Index];
 
 		// We test as floats so they round to the same precision as the sub frame
-		float     Actual   = Time / TestRate;
-		float     Expected = ExpectedSeconds[Index];
+		float     Actual   = float(Time / TestRate);
+		float     Expected = float(ExpectedSeconds[Index]);
 
 		ensureAlwaysMsgf(FMath::IsNearlyEqual(Actual, Expected, 6e-6f), TEXT("%d (+%0.3f) @ 60fps: %.9f seconds (actual) == %.9f (expected) seconds"),  Time.GetFrame().Value, Time.GetSubFrame(), Actual, Expected);
 	}
diff --git a/Engine/Source/Runtime/Windows/WindowsPlatformFeatures/Private/HighlightRecorder.cpp b/Engine/Source/Runtime/Windows/WindowsPlatformFeatures/Private/HighlightRecorder.cpp
--- a/Engine/Source/Runtime/Windows/WindowsPlatformFeatures/Private/HighlightRecorder.cpp
+++ b/Engine/Source/Runtime/Windows/WindowsPlatformFeatures/Private/HighlightRecorder.cpp
@@ -90,14 +90,14 @@ bool FHighlightRecorder::Pause(bool bPause)
 		return false;
 	}
 
-	if (bPause && PauseTimestamp == 0.0)
+	if (bPause && PauseTimestamp == 0)
 	{
 		PauseTimestamp = GetRecordingTime();
 		State = EState::Paused;
 		UE_LOG(HighlightRecorder, Log, TEXT("paused"));
 	}
 	//allow unpause to occur if we are actually paused or if we happened to pause on the same frame as recording start
-	else if (!bPause && (PauseTimestamp != 0.0 || PauseTimestamp == GetRecordingTime()))
+	else if (!bPause && (PauseTimestamp != 0 || PauseTimestamp == GetRecordingTime()))
 	{
 		FTimespan LastPausedDuration = GetRecordingTime() - PauseTimestamp;
 		TotalPausedDuration += LastPausedDuration;
diff --git a/Engine/Source/Runtime/Windows/WindowsPlatformFeatures/Private/WindowsVideoRecordingSystem.cpp b/Engine/Source/Runtime/Windows/WindowsPlatformFeatures/Private/WindowsVideoRecordingSystem.cpp
--- a/Engine/Source/Runtime/Windows/WindowsPlatformFeatures/Private/WindowsVideoRecordingSystem.cpp
+++ b/Engine/Source/Runtime/Windows/WindowsPlatformFeatures/Private/WindowsVideoRecordingSystem.cpp
@@ -261,7 +261,7 @@ bool FWindowsVideoRecordingSystem::NewRecording(const TCHAR* DestinationFileName
 	if (Recorder->GetState() == FHighlightRecorder::EState::Stopped)
 	{
 		// Call Start to initialize the internals, and pause right away
-		if (Recorder->Start(Parameters.RecordingLengthSeconds) == false)
+		if (Recorder->Start(double(Parameters.RecordingLengthSeconds)) == false)
 		{
 			// this could be if running with -nullrhi, and the Pause below will crash
 			return false;
@@ -296,7 +296,7 @@ void FWindowsVideoRecordingSystem::StartRecording()
 
 	if (Recorder->GetState() == FHighlightRecorder::EState::Stopped)
 	{
-		Recorder->Start(Parameters.RecordingLengthSeconds);
+		Recorder->Start(double(Parameters.RecordingLengthSeconds));
 	}
 	if (Recorder->GetState() == FHighlightRecorder::EState::Paused)
 	{
@@ -343,7 +343,7 @@ uint64 FWindowsVideoRecordingSystem::GetMaximumRecordingSeconds() const
 
 float FWindowsVideoRecordingSystem::GetCurrentRecordingSeconds() const
 {
-	float Ret = (FPlatformTime::Cycles64() - CurrentStartRecordingCycles + CyclesBeforePausing) * FPlatformTime::GetSecondsPerCycle();
+	float Ret = float((FPlatformTime::Cycles64() - CurrentStartRecordingCycles + CyclesBeforePausing) * FPlatformTime::GetSecondsPerCycle());
 	UE_LOG(WindowsVideoRecordingSystem, Verbose, TEXT("%s: reporting %f"), __FUNCTIONW__, Ret);
 	return Ret;
 }
@@ -371,7 +371,7 @@ void FWindowsVideoRecordingSystem::FinalizeRecording(const bool bSaveRecording,
 				FSimpleDelegateGraphTask::FDelegate::CreateRaw(this, &FWindowsVideoRecordingSystem::FinalizeCallbackOnGameThread,
 					bRes, Parameters.bAutoContinue && !bStopAutoContinue, InFullPathToFile, true),
 				TStatId(), nullptr, ENamedThreads::GameThread);
-		}, Parameters.RecordingLengthSeconds);
+		}, double(Parameters.RecordingLengthSeconds));
 
 		if (!bRet)
 		{
