Moving loose files in ViewportInteraction into Private and Public.

diff --git a/Engine/Source/Editor/ViewportInteraction/Private/Gizmo/VIStretchGizmoHandle.h b/Engine/Source/Editor/ViewportInteraction/Private/Gizmo/VIStretchGizmoHandle.h
new file mode 100644
--- /dev/null
+++ b/Engine/Source/Editor/ViewportInteraction/Private/Gizmo/VIStretchGizmoHandle.h
@@ -0,0 +1,46 @@
+// Copyright Epic Games, Inc. All Rights Reserved.
+
+#pragma once
+
+#include "CoreMinimal.h"
+#include "UObject/ObjectMacros.h"
+#include "VIGizmoHandle.h"
+#include "ViewportDragOperation.h"
+#include "VIStretchGizmoHandle.generated.h"
+
+enum class EGizmoHandleTypes : uint8;
+
+/**
+ * Gizmo handle for stretching/scaling
+ */
+UCLASS()
+class VIEWPORTINTERACTION_API UStretchGizmoHandleGroup : public UGizmoHandleGroup
+{
+	GENERATED_BODY()
+
+public:
+
+	/** Default constructor that sets up CDO properties */
+	UStretchGizmoHandleGroup();
+	
+	/** Updates the Gizmo handles */
+	virtual void UpdateGizmoHandleGroup( const FTransform& LocalToWorld, const FBox& LocalBounds, const FVector ViewLocation, const bool bAllHandlesVisible, class UActorComponent* DraggingHandle, const TArray< UActorComponent* >& HoveringOverHandles, 
+		float AnimationAlpha, float GizmoScale, const float GizmoHoverScale, const float GizmoHoverAnimationDuration, bool& bOutIsHoveringOrDraggingThisHandleGroup ) override;
+
+	/** Gets the GizmoType for this Gizmo handle */
+	virtual EGizmoHandleTypes GetHandleType() const override;
+
+	/** Returns true if this type of handle is allowed with world space gizmos */
+	virtual bool SupportsWorldCoordinateSpace() const override;
+};
+
+UCLASS()
+class UStretchGizmoHandleDragOperation : public UViewportDragOperation
+{
+	GENERATED_BODY()
+
+public:
+
+	// IViewportDragOperation
+	virtual void ExecuteDrag(struct FDraggingTransformableData& DraggingData) override;
+};
diff --git a/Engine/Source/Editor/ViewportInteraction/Private/Gizmo/VIUniformScaleGizmoHandle.h b/Engine/Source/Editor/ViewportInteraction/Private/Gizmo/VIUniformScaleGizmoHandle.h
new file mode 100644
--- /dev/null
+++ b/Engine/Source/Editor/ViewportInteraction/Private/Gizmo/VIUniformScaleGizmoHandle.h
@@ -0,0 +1,39 @@
+// Copyright Epic Games, Inc. All Rights Reserved.
+
+#pragma once
+
+#include "CoreMinimal.h"
+#include "UObject/ObjectMacros.h"
+#include "VIGizmoHandle.h"
+#include "VIUniformScaleGizmoHandle.generated.h"
+
+enum class EGizmoHandleTypes : uint8;
+
+/**
+ * Gizmo handle for uniform scaling
+ */
+UCLASS()
+class VIEWPORTINTERACTION_API UUniformScaleGizmoHandleGroup : public UGizmoHandleGroup
+{
+	GENERATED_BODY()
+
+public:
+
+	/** Default constructor that sets up CDO properties */
+	UUniformScaleGizmoHandleGroup();
+	
+	/** Updates the Gizmo handles */
+	virtual void UpdateGizmoHandleGroup( const FTransform& LocalToWorld, const FBox& LocalBounds, const FVector ViewLocation, const bool bAllHandlesVisible, class UActorComponent* DraggingHandle, const TArray< UActorComponent* >& HoveringOverHandles, 
+		float AnimationAlpha, float GizmoScale, const float GizmoHoverScale, const float GizmoHoverAnimationDuration, bool& bOutIsHoveringOrDraggingThisHandleGroup ) override;
+
+	/** Gets the GizmoType for this Gizmo handle */
+	virtual EGizmoHandleTypes GetHandleType() const override;
+
+	/** Sets if the pivot point is used as location for the handle */
+	void SetUsePivotPointAsLocation( const bool bInUsePivotAsLocation );
+
+private:
+
+	/** If the pivot point is used for the uniform scaling handle */
+	bool bUsePivotAsLocation;
+};
diff --git a/Engine/Source/Editor/ViewportInteraction/Public/ViewportInteractor.h b/Engine/Source/Editor/ViewportInteraction/Public/ViewportInteractor.h
--- a/Engine/Source/Editor/ViewportInteraction/Public/ViewportInteractor.h
+++ b/Engine/Source/Editor/ViewportInteraction/Public/ViewportInteractor.h
@@ -8,7 +8,7 @@
 #include "InputCoreTypes.h"
 #include "Engine/EngineTypes.h"
 #include "ViewportInteractionTypes.h"
-#include "Editor/ViewportInteraction/ViewportInteractorData.h"
+#include "ViewportInteractorData.h"
 #include "ViewportInteractionUtils.h"
 #include "ViewportInteractor.generated.h"
 
diff --git a/Engine/Source/Editor/ViewportInteraction/ViewportInteraction.Build.cs b/Engine/Source/Editor/ViewportInteraction/ViewportInteraction.Build.cs
--- a/Engine/Source/Editor/ViewportInteraction/ViewportInteraction.Build.cs
+++ b/Engine/Source/Editor/ViewportInteraction/ViewportInteraction.Build.cs
@@ -8,8 +8,6 @@ namespace UnrealBuildTool.Rules
     {
         public ViewportInteraction(ReadOnlyTargetRules Target) : base(Target)
         {
-			PrivateIncludePaths.Add(ModuleDirectory);
-
             PublicDependencyModuleNames.AddRange(
                 new string[] {
                     "Core",
