Virtual Scouting: Fix for VR transform gizmo not being centered on grouped objects.

Made UActorTransformer do selection normalization consistent with FActorElementLevelEditorSelectionCustomization, which has the effect of appending group actors to the end of the transformables list. (The gizmo uses the last selected actor to determine the position of its visual.)

diff --git a/Engine/Source/Editor/ViewportInteraction/Private/ActorTransformer.cpp b/Engine/Source/Editor/ViewportInteraction/Private/ActorTransformer.cpp
--- a/Engine/Source/Editor/ViewportInteraction/Private/ActorTransformer.cpp
+++ b/Engine/Source/Editor/ViewportInteraction/Private/ActorTransformer.cpp
@@ -7,6 +7,10 @@
 #include "ViewportInteractionAssetContainer.h"
 #include "ViewportInteractor.h"
 #include "Editor.h"
+#include "Editor/GroupActor.h"
+#include "ActorGroupingUtils.h"
+#include <utility>
+
 
 void UActorTransformer::Init( UViewportWorldInteraction* InitViewportWorldInteraction )
 {
@@ -59,38 +63,102 @@ void UActorTransformer::OnStopDragging(class UViewportInteractor* Interactor)
 	}
 }
 
+
+namespace UE::ViewportInteraction::Private {
+	template <typename InContainerType>
+	bool ActorHasParentInSelection(const AActor* Actor, const InContainerType& SelectionSet)
+	{
+		check(Actor);
+		for (const AActor* ParentActor = Actor->GetAttachParentActor(); ParentActor; ParentActor = ParentActor->GetAttachParentActor())
+		{
+			if (SelectionSet.Contains(ParentActor))
+			{
+				return true;
+			}
+		}
+		return false;
+	}
+}
+
+
 void UActorTransformer::OnActorSelectionChanged( UObject* ChangedObject )
 {
+	using namespace UE::ViewportInteraction::Private;
+
 	TArray<TUniquePtr<FViewportTransformable>> NewTransformables;
 
 	USelection* ActorSelectionSet = GEditor->GetSelectedActors();
 
-	static TArray<UObject*> SelectedActorObjects;
-	SelectedActorObjects.Reset();
-	ActorSelectionSet->GetSelectedObjects( AActor::StaticClass(), /* Out */ SelectedActorObjects );
+	TArray<AActor*> SelectedActors;
+	ActorSelectionSet->GetSelectedObjects(SelectedActors);
+
+	// Imitate the logic of FActorElementLevelEditorSelectionCustomization::AppendNormalizedActors
+	TArray<AActor*> NormalizedActors; // maintains order
+	TSet<AActor*> NormalizedActorSet; // quick membership tests
 
-	for( UObject* SelectedActorObject : SelectedActorObjects )
+	auto AddUniqueNormalizedActor = [&NormalizedActors, &NormalizedActorSet](AActor* InActor)
 	{
-		AActor* SelectedActor = Cast<AActor>( SelectedActorObject );
-		if( SelectedActor != nullptr )
+		if (!NormalizedActorSet.Contains(InActor))
 		{
-			// We only are able to move objects that have a root scene component
-			if( SelectedActor->GetRootComponent() != nullptr )
-			{
-				FActorViewportTransformable* Transformable = new FActorViewportTransformable();
+			NormalizedActors.Add(InActor);
+			NormalizedActorSet.Add(InActor);
+		}
+	};
 
-				NewTransformables.Add( TUniquePtr<FViewportTransformable>( Transformable ) );
+	for (AActor* SelectedActor : SelectedActors)
+	{
+		// We only are able to move objects that have a root scene component
+		if (!SelectedActor || !SelectedActor->GetRootComponent())
+		{
+			continue;
+		}
+
+		// Ensure that only parent-most actors are included
+		if (ActorHasParentInSelection(SelectedActor, SelectedActors))
+		{
+			continue;
+		}
 
-				Transformable->ActorWeakPtr = SelectedActor;
-				Transformable->StartTransform = SelectedActor->GetTransform();
-				for (UViewportInteractor* Interactor : ViewportWorldInteraction->GetInteractors())
-				{
-					if (Interactor->CanCarry())
+		// Expand groups
+		AGroupActor* ParentGroup = AGroupActor::GetRootForActor(SelectedActor, true, true);
+		if (ParentGroup && UActorGroupingUtils::IsGroupingActive())
+		{
+			// Skip if the group is already in the normalized list, since this logic will have already run
+			if (!NormalizedActorSet.Contains(ParentGroup))
+			{
+				ParentGroup->ForEachActorInGroup(
+					[AddUniqueNormalizedActor, &SelectedActors = std::as_const(SelectedActors)]
+					(AActor* InGroupedActor, AGroupActor* InGroupActor)
 					{
-						Transformable->bShouldBeCarried = true;
-						break;
-					}
-				}
+						// Check that we've not got a parent attachment within the group/selection
+						if (GroupActorHelpers::ActorHasParentInGroup(InGroupedActor, InGroupActor)
+							|| ActorHasParentInSelection(InGroupedActor, SelectedActors))
+						{
+							return;
+						}
+
+						AddUniqueNormalizedActor(InGroupedActor);
+					});
+			}
+		}
+
+		AddUniqueNormalizedActor(SelectedActor);
+	}
+
+	for (AActor* SelectedActor : NormalizedActors)
+	{
+		FActorViewportTransformable* Transformable = new FActorViewportTransformable();
+
+		NewTransformables.Add(TUniquePtr<FViewportTransformable>(Transformable));
+
+		Transformable->ActorWeakPtr = SelectedActor;
+		Transformable->StartTransform = SelectedActor->GetTransform();
+		for (UViewportInteractor* Interactor : ViewportWorldInteraction->GetInteractors())
+		{
+			if (Interactor->CanCarry())
+			{
+				Transformable->bShouldBeCarried = true;
+				break;
 			}
 		}
 	}
