Fixed mismatched copy constructor/assignment operator in DVector.

diff --git a/Engine/Source/Editor/ViewportInteraction/Private/ViewportWorldInteraction.cpp b/Engine/Source/Editor/ViewportInteraction/Private/ViewportWorldInteraction.cpp
--- a/Engine/Source/Editor/ViewportInteraction/Private/ViewportWorldInteraction.cpp
+++ b/Engine/Source/Editor/ViewportInteraction/Private/ViewportWorldInteraction.cpp
@@ -124,11 +124,6 @@ struct DVector
 	{
 	}
 
-	DVector( const FVector& V )
-		: X( V.X ), Y( V.Y ), Z( V.Z )
-	{
-	}
-
 	DVector( double InX, double InY, double InZ )
 		: X( InX ), Y( InY ), Z( InZ )
 	{
