Fix for DVector conversion constructor being accidentally deleted.

diff --git a/Engine/Source/Editor/ViewportInteraction/Private/ViewportWorldInteraction.cpp b/Engine/Source/Editor/ViewportInteraction/Private/ViewportWorldInteraction.cpp
--- a/Engine/Source/Editor/ViewportInteraction/Private/ViewportWorldInteraction.cpp
+++ b/Engine/Source/Editor/ViewportInteraction/Private/ViewportWorldInteraction.cpp
@@ -119,7 +119,7 @@ struct DVector
 	{
 	}
 
-	DVector( const DVector& V )
+	DVector( const FVector& V )
 		: X( V.X ), Y( V.Y ), Z( V.Z )
 	{
 	}
