Virtual Scouting: Fix for gizmos not loading correctly.

Don't load the UViewportInteractionAssetContainer in the APivotTransformGizmo constructor, and don't early out of the constructor if the CDO flag is set. Instead, if no other material has been specified, let UGizmoHandleGroup::UpdateHandleColor load the default materials from the asset container just before it creates the MIDs.

#rnx

diff --git a/Engine/Source/Editor/ViewportInteraction/Private/Gizmo/VIGizmoHandle.cpp b/Engine/Source/Editor/ViewportInteraction/Private/Gizmo/VIGizmoHandle.cpp
--- a/Engine/Source/Editor/ViewportInteraction/Private/Gizmo/VIGizmoHandle.cpp
+++ b/Engine/Source/Editor/ViewportInteraction/Private/Gizmo/VIGizmoHandle.cpp
@@ -7,6 +7,7 @@
 #include "VIBaseTransformGizmo.h"
 #include "ViewportWorldInteraction.h"
 #include "ViewportDragOperation.h"
+#include "ViewportInteractionAssetContainer.h"
 #include "VIGizmoHandleMeshComponent.h"
 
 UGizmoHandleGroup::UGizmoHandleGroup()
@@ -169,11 +170,23 @@ void UGizmoHandleGroup::UpdateHandleColor( const int32 AxisIndex, FGizmoHandle&
 
 	if ( !HandleMesh->GetMaterial( 0 )->IsA( UMaterialInstanceDynamic::StaticClass() ) )
 	{
+		if (!GizmoMaterial)
+		{
+			const UViewportInteractionAssetContainer* AssetContainer = UViewportWorldInteraction::LoadAssetContainer();
+			GizmoMaterial = AssetContainer->TransformGizmoMaterial;
+		}
+
 		UMaterialInstanceDynamic* MID = UMaterialInstanceDynamic::Create( GizmoMaterial, this );
 		HandleMesh->SetMaterial( 0, MID );
 	}
 	if ( !HandleMesh->GetMaterial( 1 )->IsA( UMaterialInstanceDynamic::StaticClass() ) )
 	{
+		if (!TranslucentGizmoMaterial)
+		{
+			const UViewportInteractionAssetContainer* AssetContainer = UViewportWorldInteraction::LoadAssetContainer();
+			TranslucentGizmoMaterial = AssetContainer->TranslucentTransformGizmoMaterial;
+		}
+
 		UMaterialInstanceDynamic* MID = UMaterialInstanceDynamic::Create( TranslucentGizmoMaterial, this );
 		HandleMesh->SetMaterial( 1, MID );
 	}
diff --git a/Engine/Source/Editor/ViewportInteraction/Private/Gizmo/VIPivotTransformGizmo.cpp b/Engine/Source/Editor/ViewportInteraction/Private/Gizmo/VIPivotTransformGizmo.cpp
--- a/Engine/Source/Editor/ViewportInteraction/Private/Gizmo/VIPivotTransformGizmo.cpp
+++ b/Engine/Source/Editor/ViewportInteraction/Private/Gizmo/VIPivotTransformGizmo.cpp
@@ -40,55 +40,35 @@ APivotTransformGizmo::APivotTransformGizmo() :
 	AimingAtGizmoScaleAlpha(0.0f),
 	LastDraggingHandle(nullptr)
 {
-	if (HasAnyFlags(RF_ClassDefaultObject))
-	{
-		return;
-	}
 	bIsEditorOnlyActor = true;
 
-	TStrongObjectPtr<const UViewportInteractionAssetContainer> AssetContainer(UViewportWorldInteraction::LoadAssetContainer()); 
-	UMaterialInterface* GizmoMaterial = AssetContainer->TransformGizmoMaterial;
-	UMaterialInterface* TranslucentGizmoMaterial = AssetContainer->TranslucentTransformGizmoMaterial;
-
 	UniformScaleGizmoHandleGroup = CreateDefaultSubobject<UUniformScaleGizmoHandleGroup>( TEXT( "UniformScaleHandles" ), true );
 	UniformScaleGizmoHandleGroup->SetOwningTransformGizmo(this);
-	UniformScaleGizmoHandleGroup->SetTranslucentGizmoMaterial( TranslucentGizmoMaterial );
-	UniformScaleGizmoHandleGroup->SetGizmoMaterial( GizmoMaterial );
 	UniformScaleGizmoHandleGroup->SetupAttachment( SceneComponent );
 	AllHandleGroups.Add( UniformScaleGizmoHandleGroup );
 
 	TranslationGizmoHandleGroup = CreateDefaultSubobject<UPivotTranslationGizmoHandleGroup>(TEXT("TranslationHandles"), true);
 	TranslationGizmoHandleGroup->SetOwningTransformGizmo(this);
-	TranslationGizmoHandleGroup->SetTranslucentGizmoMaterial(TranslucentGizmoMaterial);
-	TranslationGizmoHandleGroup->SetGizmoMaterial(GizmoMaterial);
 	TranslationGizmoHandleGroup->SetupAttachment( SceneComponent );
 	AllHandleGroups.Add(TranslationGizmoHandleGroup);
 
 	ScaleGizmoHandleGroup = CreateDefaultSubobject<UPivotScaleGizmoHandleGroup>( TEXT( "ScaleHandles" ), true );
 	ScaleGizmoHandleGroup->SetOwningTransformGizmo(this);
-	ScaleGizmoHandleGroup->SetTranslucentGizmoMaterial( TranslucentGizmoMaterial );
-	ScaleGizmoHandleGroup->SetGizmoMaterial( GizmoMaterial );
 	ScaleGizmoHandleGroup->SetupAttachment( SceneComponent );
 	AllHandleGroups.Add( ScaleGizmoHandleGroup );
 
 	PlaneTranslationGizmoHandleGroup = CreateDefaultSubobject<UPivotPlaneTranslationGizmoHandleGroup>( TEXT( "PlaneTranslationHandles" ), true );
 	PlaneTranslationGizmoHandleGroup->SetOwningTransformGizmo(this);
-	PlaneTranslationGizmoHandleGroup->SetTranslucentGizmoMaterial( TranslucentGizmoMaterial );
-	PlaneTranslationGizmoHandleGroup->SetGizmoMaterial( GizmoMaterial );
 	PlaneTranslationGizmoHandleGroup->SetupAttachment( SceneComponent );
 	AllHandleGroups.Add( PlaneTranslationGizmoHandleGroup );
 
 	RotationGizmoHandleGroup = CreateDefaultSubobject<UPivotRotationGizmoHandleGroup>( TEXT( "RotationHandles" ), true );
 	RotationGizmoHandleGroup->SetOwningTransformGizmo(this);
-	RotationGizmoHandleGroup->SetTranslucentGizmoMaterial( TranslucentGizmoMaterial );
-	RotationGizmoHandleGroup->SetGizmoMaterial( GizmoMaterial );
 	RotationGizmoHandleGroup->SetupAttachment( SceneComponent );
 	AllHandleGroups.Add( RotationGizmoHandleGroup );
 
 	StretchGizmoHandleGroup = CreateDefaultSubobject<UStretchGizmoHandleGroup>( TEXT( "StretchHandles" ), true );
 	StretchGizmoHandleGroup->SetOwningTransformGizmo(this);
-	StretchGizmoHandleGroup->SetTranslucentGizmoMaterial( TranslucentGizmoMaterial );
-	StretchGizmoHandleGroup->SetGizmoMaterial( GizmoMaterial );
 	StretchGizmoHandleGroup->SetShowOnUniversalGizmo( false );
 	StretchGizmoHandleGroup->SetupAttachment( SceneComponent );
 	AllHandleGroups.Add( StretchGizmoHandleGroup );
