Updated headers using UnrealCodeFixup to make sure dllstorage is on methods/staticvars instead of types. Part 5/n

diff --git a/Engine/Source/Editor/VREditor/Public/VREditorModeBase.h b/Engine/Source/Editor/VREditor/Public/VREditorModeBase.h
--- a/Engine/Source/Editor/VREditor/Public/VREditorModeBase.h
+++ b/Engine/Source/Editor/VREditor/Public/VREditorModeBase.h
@@ -8,14 +8,16 @@
 #include "ShowFlags.h"
 #include "VREditorModeBase.generated.h"
 
+#define UE_API VREDITOR_API
+
 
 class SLevelViewport;
 
 
 /**
  */
-UCLASS(Abstract, BlueprintType, Blueprintable, Transient)
-class VREDITOR_API UVREditorModeBase : public UEditorWorldExtension
+UCLASS(MinimalAPI, Abstract, BlueprintType, Blueprintable, Transient)
+class UVREditorModeBase : public UEditorWorldExtension
 {
 	GENERATED_BODY()
 
@@ -24,12 +26,12 @@ protected:
 
 public:
 	//~ Begin UEditorWorldExtension interface
-	virtual void Init() override;
-	virtual void Shutdown() override;
+	UE_API virtual void Init() override;
+	UE_API virtual void Shutdown() override;
 	//~ End UEditorWorldExtension interface
 
-	virtual void Enter();
-	virtual void Exit(bool bShouldDisableStereo);
+	UE_API virtual void Enter();
+	UE_API virtual void Exit(bool bShouldDisableStereo);
 
 	/** Sets whether we should actually use an HMD. Call this before activating VR mode */
 	virtual void SetActuallyUsingVR(bool bShouldActuallyUseVR) { bActuallyUsingVR = bShouldActuallyUseVR; }
@@ -54,8 +56,8 @@ public:
 	FOnVRModeEntryComplete& OnVRModeEntryComplete() { return OnVRModeEntryCompleteEvent; }
 
 	// TODO: Deprecate these in favor of replacements below; viewport may be null.
-	const class SLevelViewport& GetLevelViewportPossessedForVR() const;
-	class SLevelViewport& GetLevelViewportPossessedForVR();
+	UE_API const class SLevelViewport& GetLevelViewportPossessedForVR() const;
+	UE_API class SLevelViewport& GetLevelViewportPossessedForVR();
 
 	TSharedPtr<SLevelViewport> GetVrLevelViewport() { return VREditorLevelViewportWeakPtr.Pin(); }
 	TSharedPtr<const SLevelViewport> GetVrLevelViewport() const { return VREditorLevelViewportWeakPtr.Pin(); }
@@ -69,16 +71,16 @@ protected:
 	const FBaseSavedEditorState& SavedEditorStateChecked() const { check(SavedEditorStatePtr); return *SavedEditorStatePtr; }
 
 	/** Start using the viewport passed */
-	virtual void StartViewport(TSharedPtr<SLevelViewport> Viewport);
+	UE_API virtual void StartViewport(TSharedPtr<SLevelViewport> Viewport);
 
 	/** Close the current viewport */
-	virtual void CloseViewport(bool bShouldDisableStereo);
+	UE_API virtual void CloseViewport(bool bShouldDisableStereo);
 
-	virtual void StereoViewportSetup(TSharedRef<SLevelViewport> Viewport);
-	virtual void StereoViewportShutdown(TSharedRef<SLevelViewport> Viewport);
+	UE_API virtual void StereoViewportSetup(TSharedRef<SLevelViewport> Viewport);
+	UE_API virtual void StereoViewportShutdown(TSharedRef<SLevelViewport> Viewport);
 
-	virtual void EnableStereo();
-	virtual void DisableStereo();
+	UE_API virtual void EnableStereo();
+	UE_API virtual void DisableStereo();
 
 protected:
 	/** True if we're in using an actual HMD in this mode, or false if we're "faking" VR mode for testing */
@@ -108,3 +110,5 @@ protected:
 
 	TSharedPtr<FBaseSavedEditorState> SavedEditorStatePtr;
 };
+
+#undef UE_API
diff --git a/Engine/Source/Editor/VREditor/Public/VRModeSettings.h b/Engine/Source/Editor/VREditor/Public/VRModeSettings.h
--- a/Engine/Source/Editor/VREditor/Public/VRModeSettings.h
+++ b/Engine/Source/Editor/VREditor/Public/VRModeSettings.h
@@ -8,6 +8,8 @@
 #include "VISettings.h"
 #include "VRModeSettings.generated.h"
 
+#define UE_API VREDITOR_API
+
 
 class AVREditorTeleporter;
 class UVREditorInteractor;
@@ -27,15 +29,15 @@ enum class EInteractorHand : uint8
 /**
 * Implements the settings for VR Mode.
 */
-UCLASS(config = EditorSettings)
-class VREDITOR_API UVRModeSettings : public UVISettings
+UCLASS(MinimalAPI, config = EditorSettings)
+class UVRModeSettings : public UVISettings
 {
 	GENERATED_BODY()
 
 public:
 
 	/** Default constructor that sets up CDO properties */
-	UVRModeSettings();
+	UE_API UVRModeSettings();
 
 	/**If true, wearing a Vive or Oculus Rift headset will automatically enter VR Editing mode */
 	UPROPERTY(EditAnywhere, config, Category = "General", meta = (DisplayName = "Enable VR Mode Auto-Entry"))
@@ -94,6 +96,8 @@ public:
 	TSoftClassPtr<UVREditorModeBase> ModeClass;
 
 #if WITH_EDITOR
-	virtual void PostEditChangeProperty(struct FPropertyChangedEvent& PropertyChangedEvent) override;
+	UE_API virtual void PostEditChangeProperty(struct FPropertyChangedEvent& PropertyChangedEvent) override;
 #endif
 };
+
+#undef UE_API
diff --git a/Engine/Source/Editor/VREditor/Public/VRScoutingInteractor.h b/Engine/Source/Editor/VREditor/Public/VRScoutingInteractor.h
--- a/Engine/Source/Editor/VREditor/Public/VRScoutingInteractor.h
+++ b/Engine/Source/Editor/VREditor/Public/VRScoutingInteractor.h
@@ -9,6 +9,8 @@
 #include "VREditorInteractor.h"
 #include "VRScoutingInteractor.generated.h"
 
+#define UE_API VREDITOR_API
+
 class AActor;
 class UStaticMesh;
 class UStaticMeshSocket;
@@ -17,36 +19,36 @@ class UInputComponent;
 /**
  * Represents the interactor in the world
  */
-UCLASS(Abstract)
-class VREDITOR_API UVRScoutingInteractor : public UVREditorInteractor
+UCLASS(MinimalAPI, Abstract)
+class UVRScoutingInteractor : public UVREditorInteractor
 {
 	GENERATED_BODY()
 
 public:
 
 	/** Default constructor */
-	UVRScoutingInteractor();
+	UE_API UVRScoutingInteractor();
 
 	/** Gets the trackpad slide delta */
-	virtual float GetSlideDelta_Implementation() const override;
+	UE_API virtual float GetSlideDelta_Implementation() const override;
 
 	/** Sets up all components */
-	virtual void SetupComponent_Implementation(AActor* OwningActor) override;
+	UE_API virtual void SetupComponent_Implementation(AActor* OwningActor) override;
 
 	// IViewportInteractorInterface overrides
-	virtual void Shutdown_Implementation() override;
+	UE_API virtual void Shutdown_Implementation() override;
 
 	/** Sets the gizmo mode for selected object */
 	UFUNCTION(BlueprintCallable, Category = "Scouting")
-	void SetGizmoMode(EGizmoHandleTypes InGizmoMode);
+	UE_API void SetGizmoMode(EGizmoHandleTypes InGizmoMode);
 
 	/** Gets the gizmo mode for selected object */
 	UFUNCTION(BlueprintCallable, Category = "Scouting")
-	EGizmoHandleTypes GetGizmoMode() const;
+	UE_API EGizmoHandleTypes GetGizmoMode() const;
 
 	/** Gets all actors that are selected in the world editor */
 	UFUNCTION(BlueprintCallable, Category = "Scouting")
-	static TArray<AActor*> GetSelectedActors();
+	static UE_API TArray<AActor*> GetSelectedActors();
 
 	/** Shown in Navigation mode */
 	UPROPERTY(Category = Interactor, EditAnywhere, BlueprintReadOnly)
@@ -60,7 +62,7 @@ public:
 	}
 
 	UFUNCTION(BlueprintSetter, Category = "Input|Editor")
-	void SetReceivesEditorInput(bool bInValue);
+	UE_API void SetReceivesEditorInput(bool bInValue);
 
 	UFUNCTION(BlueprintGetter, BlueprintPure, Category = "Input|Editor")
 	bool GetReceivesEditorInput() const
@@ -71,10 +73,10 @@ public:
 private:
 
 	/** Creates the EditorOnlyInputComponent if it does not already exist and registers all subobject callbacks to it */
-	void CreateEditorInput();
+	UE_API void CreateEditorInput();
 
 	/** Removes the EditorOnlyInputComponent from this object */
-	void RemoveEditorInput();
+	UE_API void RemoveEditorInput();
 
 	UPROPERTY(Transient, DuplicateTransient)
 	TObjectPtr<UInputComponent> EditorOnlyInputComponent;
@@ -84,3 +86,5 @@ private:
 	bool bReceivesEditorInput = true;
 
 };
+
+#undef UE_API
diff --git a/Engine/Source/Editor/ViewportInteraction/Private/Gizmo/VIStretchGizmoHandle.h b/Engine/Source/Editor/ViewportInteraction/Private/Gizmo/VIStretchGizmoHandle.h
--- a/Engine/Source/Editor/ViewportInteraction/Private/Gizmo/VIStretchGizmoHandle.h
+++ b/Engine/Source/Editor/ViewportInteraction/Private/Gizmo/VIStretchGizmoHandle.h
@@ -8,30 +8,32 @@
 #include "ViewportDragOperation.h"
 #include "VIStretchGizmoHandle.generated.h"
 
+#define UE_API VIEWPORTINTERACTION_API
+
 enum class EGizmoHandleTypes : uint8;
 
 /**
  * Gizmo handle for stretching/scaling
  */
-UCLASS()
-class VIEWPORTINTERACTION_API UStretchGizmoHandleGroup : public UGizmoHandleGroup
+UCLASS(MinimalAPI)
+class UStretchGizmoHandleGroup : public UGizmoHandleGroup
 {
 	GENERATED_BODY()
 
 public:
 
 	/** Default constructor that sets up CDO properties */
-	UStretchGizmoHandleGroup();
+	UE_API UStretchGizmoHandleGroup();
 	
 	/** Updates the Gizmo handles */
-	virtual void UpdateGizmoHandleGroup( const FTransform& LocalToWorld, const FBox& LocalBounds, const FVector ViewLocation, const bool bAllHandlesVisible, class UActorComponent* DraggingHandle, const TArray< UActorComponent* >& HoveringOverHandles, 
+	UE_API virtual void UpdateGizmoHandleGroup( const FTransform& LocalToWorld, const FBox& LocalBounds, const FVector ViewLocation, const bool bAllHandlesVisible, class UActorComponent* DraggingHandle, const TArray< UActorComponent* >& HoveringOverHandles, 
 		float AnimationAlpha, float GizmoScale, const float GizmoHoverScale, const float GizmoHoverAnimationDuration, bool& bOutIsHoveringOrDraggingThisHandleGroup ) override;
 
 	/** Gets the GizmoType for this Gizmo handle */
-	virtual EGizmoHandleTypes GetHandleType() const override;
+	UE_API virtual EGizmoHandleTypes GetHandleType() const override;
 
 	/** Returns true if this type of handle is allowed with world space gizmos */
-	virtual bool SupportsWorldCoordinateSpace() const override;
+	UE_API virtual bool SupportsWorldCoordinateSpace() const override;
 };
 
 UCLASS()
@@ -44,3 +46,5 @@ public:
 	// IViewportDragOperation
 	virtual void ExecuteDrag(struct FDraggingTransformableData& DraggingData) override;
 };
+
+#undef UE_API
diff --git a/Engine/Source/Editor/ViewportInteraction/Private/Gizmo/VIUniformScaleGizmoHandle.h b/Engine/Source/Editor/ViewportInteraction/Private/Gizmo/VIUniformScaleGizmoHandle.h
--- a/Engine/Source/Editor/ViewportInteraction/Private/Gizmo/VIUniformScaleGizmoHandle.h
+++ b/Engine/Source/Editor/ViewportInteraction/Private/Gizmo/VIUniformScaleGizmoHandle.h
@@ -7,33 +7,37 @@
 #include "VIGizmoHandle.h"
 #include "VIUniformScaleGizmoHandle.generated.h"
 
+#define UE_API VIEWPORTINTERACTION_API
+
 enum class EGizmoHandleTypes : uint8;
 
 /**
  * Gizmo handle for uniform scaling
  */
-UCLASS()
-class VIEWPORTINTERACTION_API UUniformScaleGizmoHandleGroup : public UGizmoHandleGroup
+UCLASS(MinimalAPI)
+class UUniformScaleGizmoHandleGroup : public UGizmoHandleGroup
 {
 	GENERATED_BODY()
 
 public:
 
 	/** Default constructor that sets up CDO properties */
-	UUniformScaleGizmoHandleGroup();
+	UE_API UUniformScaleGizmoHandleGroup();
 	
 	/** Updates the Gizmo handles */
-	virtual void UpdateGizmoHandleGroup( const FTransform& LocalToWorld, const FBox& LocalBounds, const FVector ViewLocation, const bool bAllHandlesVisible, class UActorComponent* DraggingHandle, const TArray< UActorComponent* >& HoveringOverHandles, 
+	UE_API virtual void UpdateGizmoHandleGroup( const FTransform& LocalToWorld, const FBox& LocalBounds, const FVector ViewLocation, const bool bAllHandlesVisible, class UActorComponent* DraggingHandle, const TArray< UActorComponent* >& HoveringOverHandles, 
 		float AnimationAlpha, float GizmoScale, const float GizmoHoverScale, const float GizmoHoverAnimationDuration, bool& bOutIsHoveringOrDraggingThisHandleGroup ) override;
 
 	/** Gets the GizmoType for this Gizmo handle */
-	virtual EGizmoHandleTypes GetHandleType() const override;
+	UE_API virtual EGizmoHandleTypes GetHandleType() const override;
 
 	/** Sets if the pivot point is used as location for the handle */
-	void SetUsePivotPointAsLocation( const bool bInUsePivotAsLocation );
+	UE_API void SetUsePivotPointAsLocation( const bool bInUsePivotAsLocation );
 
 private:
 
 	/** If the pivot point is used for the uniform scaling handle */
 	bool bUsePivotAsLocation;
 };
+
+#undef UE_API
diff --git a/Engine/Source/Editor/ViewportInteraction/Public/ActorTransformer.h b/Engine/Source/Editor/ViewportInteraction/Public/ActorTransformer.h
--- a/Engine/Source/Editor/ViewportInteraction/Public/ActorTransformer.h
+++ b/Engine/Source/Editor/ViewportInteraction/Public/ActorTransformer.h
@@ -5,28 +5,31 @@
 #include "ViewportTransformer.h"
 #include "ActorTransformer.generated.h"
 
+#define UE_API VIEWPORTINTERACTION_API
 
-UCLASS()
-class VIEWPORTINTERACTION_API UActorTransformer : public UViewportTransformer
+
+UCLASS(MinimalAPI)
+class UActorTransformer : public UViewportTransformer
 {
 	GENERATED_BODY()
 
 public:
 
 	// UViewportTransformer overrides
-	virtual void Init( class UViewportWorldInteraction* InitViewportWorldInteraction ) override;
-	virtual void Shutdown() override;
+	UE_API virtual void Init( class UViewportWorldInteraction* InitViewportWorldInteraction ) override;
+	UE_API virtual void Shutdown() override;
 	virtual bool CanAlignToActors() const override
 	{
 		return true;
 	}
-	virtual void OnStartDragging(class UViewportInteractor* Interactor) override;
-	virtual void OnStopDragging(class UViewportInteractor* Interactor) override;
+	UE_API virtual void OnStartDragging(class UViewportInteractor* Interactor) override;
+	UE_API virtual void OnStopDragging(class UViewportInteractor* Interactor) override;
 
 protected:
 
 	/** Called when selection changes in the level */
-	void OnActorSelectionChanged( UObject* ChangedObject );
+	UE_API void OnActorSelectionChanged( UObject* ChangedObject );
 
 };
 
+#undef UE_API
diff --git a/Engine/Source/Editor/ViewportInteraction/Public/ActorViewportTransformable.h b/Engine/Source/Editor/ViewportInteraction/Public/ActorViewportTransformable.h
--- a/Engine/Source/Editor/ViewportInteraction/Public/ActorViewportTransformable.h
+++ b/Engine/Source/Editor/ViewportInteraction/Public/ActorViewportTransformable.h
@@ -6,10 +6,12 @@
 #include "UObject/WeakObjectPtrTemplates.h"
 #include "UObject/WeakObjectPtr.h"
 
+#define UE_API VIEWPORTINTERACTION_API
+
 /**
  * A transformable actor
  */
-class VIEWPORTINTERACTION_API FActorViewportTransformable : public FViewportTransformable
+class FActorViewportTransformable : public FViewportTransformable
 {
 
 public:
@@ -22,14 +24,14 @@ public:
 	}
 
 	// FViewportTransformable overrides
-	virtual const FTransform GetTransform() const override;
-	virtual void ApplyTransform( const FTransform& NewTransform, const bool bSweep ) override;
-	virtual FBox BuildBoundingBox( const FTransform& BoundingBoxToWorld ) const override;
-	virtual bool IsPhysicallySimulated() const override;
-	virtual bool ShouldBeCarried() const override;
-	virtual void SetLinearVelocity( const FVector& NewVelocity ) override;
-	virtual FVector GetLinearVelocity() const override;
-	virtual void UpdateIgnoredActorList( TArray<class AActor*>& IgnoredActors ) override;
+	UE_API virtual const FTransform GetTransform() const override;
+	UE_API virtual void ApplyTransform( const FTransform& NewTransform, const bool bSweep ) override;
+	UE_API virtual FBox BuildBoundingBox( const FTransform& BoundingBoxToWorld ) const override;
+	UE_API virtual bool IsPhysicallySimulated() const override;
+	UE_API virtual bool ShouldBeCarried() const override;
+	UE_API virtual void SetLinearVelocity( const FVector& NewVelocity ) override;
+	UE_API virtual FVector GetLinearVelocity() const override;
+	UE_API virtual void UpdateIgnoredActorList( TArray<class AActor*>& IgnoredActors ) override;
 
 	/** The actual actor object */
 	TWeakObjectPtr<class AActor> ActorWeakPtr;
@@ -38,3 +40,4 @@ public:
 };
 
 
+#undef UE_API
diff --git a/Engine/Source/Editor/ViewportInteraction/Public/MouseCursorInteractor.h b/Engine/Source/Editor/ViewportInteraction/Public/MouseCursorInteractor.h
--- a/Engine/Source/Editor/ViewportInteraction/Public/MouseCursorInteractor.h
+++ b/Engine/Source/Editor/ViewportInteraction/Public/MouseCursorInteractor.h
@@ -7,29 +7,31 @@
 #include "ViewportInteractor.h"
 #include "MouseCursorInteractor.generated.h"
 
+#define UE_API VIEWPORTINTERACTION_API
+
 /**
  * Viewport interactor for a mouse cursor
  */
-UCLASS()
-class VIEWPORTINTERACTION_API UMouseCursorInteractor : public UViewportInteractor
+UCLASS(MinimalAPI)
+class UMouseCursorInteractor : public UViewportInteractor
 {
 	GENERATED_BODY()
 
 public:
 
 	/** Default constructor. */
-	UMouseCursorInteractor();
+	UE_API UMouseCursorInteractor();
 
 	/** Initialize default values */
-	void Init();
+	UE_API void Init();
 
 	// ViewportInteractor overrides
-	virtual void PollInput() override;
-	virtual bool IsModifierPressed() const override;
+	UE_API virtual void PollInput() override;
+	UE_API virtual bool IsModifierPressed() const override;
 
 protected:
 
-	virtual bool AllowLaserSmoothing() const;
+	UE_API virtual bool AllowLaserSmoothing() const;
 
 private:
 
@@ -37,3 +39,5 @@ private:
 	bool bIsControlKeyPressed;
 
 };
+
+#undef UE_API
diff --git a/Engine/Source/Editor/ViewportInteraction/Public/VIBaseTransformGizmo.h b/Engine/Source/Editor/ViewportInteraction/Public/VIBaseTransformGizmo.h
--- a/Engine/Source/Editor/ViewportInteraction/Public/VIBaseTransformGizmo.h
+++ b/Engine/Source/Editor/ViewportInteraction/Public/VIBaseTransformGizmo.h
@@ -9,6 +9,8 @@
 #include "ViewportInteractionTypes.h"
 #include "VIBaseTransformGizmo.generated.h"
 
+#define UE_API VIEWPORTINTERACTION_API
+
 class UMaterialInterface;
 
 UENUM()
@@ -37,15 +39,15 @@ struct FTransformGizmoMeasurement
 /**
  * Base class for transform gizmo
  */
-UCLASS( Abstract, NotPlaceable )
-class VIEWPORTINTERACTION_API ABaseTransformGizmo : public AActor
+UCLASS(MinimalAPI,  Abstract, NotPlaceable )
+class ABaseTransformGizmo : public AActor
 {
 	GENERATED_BODY()
 
 public:
 	
 	/** Default constructor that sets up CDO properties */
-	ABaseTransformGizmo();
+	UE_API ABaseTransformGizmo();
 
 	//~ Begin AActor interface
 	virtual bool IsEditorOnly() const final
@@ -55,21 +57,21 @@ public:
 	//~ End AActor interface
 
 	/** Call this when new objects become selected.  This triggers an animation transition. */
-	void OnNewObjectsSelected();
+	UE_API void OnNewObjectsSelected();
 
 	/** Called by the world interaction system when one of our components is dragged by the user to find out
 	    what type of interaction to do.  If null is passed in then we'll treat it as dragging the whole object 
 		(rather than a specific axis/handle) */
-	class UViewportDragOperationComponent* GetInteractionType( UActorComponent* DraggedComponent, TOptional<FTransformGizmoHandlePlacement>& OutHandlePlacement );
+	UE_API class UViewportDragOperationComponent* GetInteractionType( UActorComponent* DraggedComponent, TOptional<FTransformGizmoHandlePlacement>& OutHandlePlacement );
 
 	/** Updates the animation with the current time and selected time */
-	float GetAnimationAlpha();
+	UE_API float GetAnimationAlpha();
 
 	/** Sets the owner */
-	void SetOwnerWorldInteraction( class UViewportWorldInteraction* InWorldInteraction );
+	UE_API void SetOwnerWorldInteraction( class UViewportWorldInteraction* InWorldInteraction );
 
 	/** Gets the owner */
-	class UViewportWorldInteraction* GetOwnerWorldInteraction() const;
+	UE_API class UViewportWorldInteraction* GetOwnerWorldInteraction() const;
 
 	/** Called by the world interaction system after we've been spawned into the world, to allow
 	    us to create components and set everything up nicely for the selected objects that we'll be
@@ -79,18 +81,18 @@ public:
 
 
 	/** Gets the current gizmo handle type */
-	EGizmoHandleTypes GetGizmoType() const;
+	UE_API EGizmoHandleTypes GetGizmoType() const;
 
 protected:
 
 	/** Static: Given a bounding box and information about which edge to query, returns the vertex positions of that edge */
-	static void GetBoundingBoxEdge( const FBox& Box, const int32 AxisIndex, const int32 EdgeIndex, FVector& OutVertex0, FVector& OutVertex1 );
+	static UE_API void GetBoundingBoxEdge( const FBox& Box, const int32 AxisIndex, const int32 EdgeIndex, FVector& OutVertex0, FVector& OutVertex1 );
 
 	/** Updates the visibility of all the handles */
-	void UpdateHandleVisibility(const EGizmoHandleTypes InGizmoType, const ECoordSystem InGizmoCoordinateSpace, bool bInAllHandlesVisible, UActorComponent* DraggingHandle);
+	UE_API void UpdateHandleVisibility(const EGizmoHandleTypes InGizmoType, const ECoordSystem InGizmoCoordinateSpace, bool bInAllHandlesVisible, UActorComponent* DraggingHandle);
 
 	/** Gets if the gizmo shows measurement texts */
-	bool GetShowMeasurementText() const;
+	UE_API bool GetShowMeasurementText() const;
 
 	/** Real time that the gizmo was last attached to a selected set of objects.  This is used for animation transitions */
 	FTimespan SelectedAtTime;
@@ -110,3 +112,5 @@ protected:
 	/** Current gizmo type */
 	EGizmoHandleTypes GizmoType;
 };	
+
+#undef UE_API
diff --git a/Engine/Source/Editor/ViewportInteraction/Public/VIGizmoHandle.h b/Engine/Source/Editor/ViewportInteraction/Public/VIGizmoHandle.h
--- a/Engine/Source/Editor/ViewportInteraction/Public/VIGizmoHandle.h
+++ b/Engine/Source/Editor/ViewportInteraction/Public/VIGizmoHandle.h
@@ -9,13 +9,15 @@
 #include "ViewportInteractionTypes.h"
 #include "VIGizmoHandle.generated.h"
 
+#define UE_API VIEWPORTINTERACTION_API
+
 class UMaterialInterface;
 class UStaticMesh;
 enum class EGizmoHandleTypes : uint8;
 class UActorComponent;
 
 USTRUCT()
-struct VIEWPORTINTERACTION_API FGizmoHandle
+struct FGizmoHandle
 {
 	GENERATED_BODY()
 
@@ -30,51 +32,51 @@ struct VIEWPORTINTERACTION_API FGizmoHandle
 /**
  * Base class for gizmo handles
  */
-UCLASS( ABSTRACT )
-class VIEWPORTINTERACTION_API UGizmoHandleGroup : public USceneComponent
+UCLASS(MinimalAPI,  ABSTRACT )
+class UGizmoHandleGroup : public USceneComponent
 {
 	GENERATED_BODY()
 
 public:
 
 	/** Default constructor that sets up CDO properties */
-	UGizmoHandleGroup();
+	UE_API UGizmoHandleGroup();
 	
 	/** Given the unique index, makes a handle */
-	FTransformGizmoHandlePlacement MakeHandlePlacementForIndex( const int32 HandleIndex ) const;
+	UE_API FTransformGizmoHandlePlacement MakeHandlePlacementForIndex( const int32 HandleIndex ) const;
 
 	/** Makes a unique index for a handle */
-	int32 MakeHandleIndex( const FTransformGizmoHandlePlacement HandlePlacement ) const;
+	UE_API int32 MakeHandleIndex( const FTransformGizmoHandlePlacement HandlePlacement ) const;
 
 	/** Makes up a name string for a handle */
-	FString MakeHandleName( const FTransformGizmoHandlePlacement HandlePlacement ) const;
+	UE_API FString MakeHandleName( const FTransformGizmoHandlePlacement HandlePlacement ) const;
 
 	/** Static: Given an axis (0-2) and a facing direction, returns the vector normal for that axis */
-	static FVector GetAxisVector( const int32 AxisIndex, const ETransformGizmoHandleDirection HandleDirection );
+	static UE_API FVector GetAxisVector( const int32 AxisIndex, const ETransformGizmoHandleDirection HandleDirection );
 
 	/** Updates the Gizmo handles, needs to be implemented by derived classes */
-	virtual void UpdateGizmoHandleGroup( const FTransform& LocalToWorld, const FBox& LocalBounds, const FVector ViewLocation, const bool bAllHandlesVisible, class UActorComponent* DraggingHandle, const TArray< UActorComponent* >& HoveringOverHandles, 
+	UE_API virtual void UpdateGizmoHandleGroup( const FTransform& LocalToWorld, const FBox& LocalBounds, const FVector ViewLocation, const bool bAllHandlesVisible, class UActorComponent* DraggingHandle, const TArray< UActorComponent* >& HoveringOverHandles, 
 		float AnimationAlpha, float GizmoScale, const float GizmoHoverScale, const float GizmoHoverAnimationDuration, bool& bOutIsHoveringOrDraggingThisHandleGroup );
 
 	/** Default setting the visibility and collision for all the handles in this group */
-	void UpdateVisibilityAndCollision(const EGizmoHandleTypes GizmoType, const ECoordSystem GizmoCoordinateSpace, const bool bAllHandlesVisible, const bool bAllowRotationAndScaleHandles, UActorComponent* DraggingHandle);
+	UE_API void UpdateVisibilityAndCollision(const EGizmoHandleTypes GizmoType, const ECoordSystem GizmoCoordinateSpace, const bool bAllHandlesVisible, const bool bAllowRotationAndScaleHandles, UActorComponent* DraggingHandle);
 
-	class UViewportDragOperationComponent* GetDragOperationComponent();
+	UE_API class UViewportDragOperationComponent* GetDragOperationComponent();
 
 	/** Finds the index of DraggedMesh in HandleMeshes */
-	virtual int32 GetDraggedHandleIndex( class UStaticMeshComponent* DraggedMesh );
+	UE_API virtual int32 GetDraggedHandleIndex( class UStaticMeshComponent* DraggedMesh );
 
 	/** Sets the Gizmo material */
-	void SetGizmoMaterial( UMaterialInterface* Material );
+	UE_API void SetGizmoMaterial( UMaterialInterface* Material );
 	
 	/** Sets the translucent Gizmo material */
-	void SetTranslucentGizmoMaterial( UMaterialInterface* Material );
+	UE_API void SetTranslucentGizmoMaterial( UMaterialInterface* Material );
 
 	/** Gets all the handles */
-	TArray< FGizmoHandle >& GetHandles();
+	UE_API TArray< FGizmoHandle >& GetHandles();
 
 	/** Gets the GizmoType for this Gizmo handle */
-	virtual EGizmoHandleTypes GetHandleType() const;
+	UE_API virtual EGizmoHandleTypes GetHandleType() const;
 
 	/** Returns true if this type of handle is allowed with world space gizmos */
 	virtual bool SupportsWorldCoordinateSpace() const
@@ -83,29 +85,29 @@ public:
 	}
 
 	/** Sets if this handlegroup will be visible with the universal gizmo */
-	void SetShowOnUniversalGizmo( const bool bShowHandleUniversal );
+	UE_API void SetShowOnUniversalGizmo( const bool bShowHandleUniversal );
 	
 	/** Gets if this handlegroup will be visible with the universal gizmo */
-	bool GetShowOnUniversalGizmo() const;
+	UE_API bool GetShowOnUniversalGizmo() const;
 
 	/** Sets the owning transform gizmo for this handle group*/
-	void SetOwningTransformGizmo(class ABaseTransformGizmo* TransformGizmo); 
+	UE_API void SetOwningTransformGizmo(class ABaseTransformGizmo* TransformGizmo); 
 
 protected:
 	/** Updates the colors of the dynamic material instances for the handle passed using its axis index */	
-	void UpdateHandleColor( const int32 AxisIndex, FGizmoHandle& Handle, class UActorComponent* DraggingHandle, const TArray< UActorComponent* >& HoveringOverHandles );
+	UE_API void UpdateHandleColor( const int32 AxisIndex, FGizmoHandle& Handle, class UActorComponent* DraggingHandle, const TArray< UActorComponent* >& HoveringOverHandles );
 
 	/** Helper function to create gizmo handle meshes */
-	class UGizmoHandleMeshComponent* CreateMeshHandle( class UStaticMesh* HandleMesh, const FString& ComponentName );
+	UE_API class UGizmoHandleMeshComponent* CreateMeshHandle( class UStaticMesh* HandleMesh, const FString& ComponentName );
 
 	/** Creates handle meshcomponent and adds it to the Handles list */
-	class UGizmoHandleMeshComponent* CreateAndAddMeshHandle( class UStaticMesh* HandleMesh, const FString& ComponentName, const FTransformGizmoHandlePlacement& HandlePlacement );
+	UE_API class UGizmoHandleMeshComponent* CreateAndAddMeshHandle( class UStaticMesh* HandleMesh, const FString& ComponentName, const FTransformGizmoHandlePlacement& HandlePlacement );
 
 	/** Adds the HandleMeshComponent to the Handles list */
-	void AddMeshToHandles( class UGizmoHandleMeshComponent* HandleMeshComponent, const FTransformGizmoHandlePlacement& HandlePlacement );
+	UE_API void AddMeshToHandles( class UGizmoHandleMeshComponent* HandleMeshComponent, const FTransformGizmoHandlePlacement& HandlePlacement );
 
 	/** Gets the handleplacement axes */
-	FTransformGizmoHandlePlacement GetHandlePlacement( const int32 X, const int32 Y, const int32 Z ) const;
+	UE_API FTransformGizmoHandlePlacement GetHandlePlacement( const int32 X, const int32 Y, const int32 Z ) const;
 
 	/** Gizmo material (opaque) */
 	UPROPERTY()
@@ -129,7 +131,7 @@ protected:
 private:
 
 	/** Updates the hover animation for the HoveringOverHandles */
-	void UpdateHoverAnimation( class UActorComponent* DraggingHandle, const TArray< UActorComponent* >& HoveringOverHandles, const float GizmoHoverAnimationDuration, bool& bOutIsHoveringOrDraggingThisHandleGroup );
+	UE_API void UpdateHoverAnimation( class UActorComponent* DraggingHandle, const TArray< UActorComponent* >& HoveringOverHandles, const float GizmoHoverAnimationDuration, bool& bOutIsHoveringOrDraggingThisHandleGroup );
 
 	/** If this handlegroup will be visible with the universal gizmo */
 	bool bShowOnUniversalGizmo;
@@ -138,20 +140,22 @@ private:
 /**
  * Base class for gizmo handles on axis
  */
-UCLASS(ABSTRACT)
-class VIEWPORTINTERACTION_API UAxisGizmoHandleGroup : public UGizmoHandleGroup
+UCLASS(MinimalAPI, ABSTRACT)
+class UAxisGizmoHandleGroup : public UGizmoHandleGroup
 {
 	GENERATED_BODY()
 
 protected:
 
 	/** Creates mesh for every axis */
-	void CreateHandles(UStaticMesh* HandleMesh, const FString& HandleComponentName);
+	UE_API void CreateHandles(UStaticMesh* HandleMesh, const FString& HandleComponentName);
 	
 	/** 
 	 * Calculates the transforms of meshes of this handlegroup 
 	 * @param HandleToCenter - The offset from the center of the gizmo in roomspace
 	 */
-	void UpdateHandlesRelativeTransformOnAxis( const FTransform& HandleToCenter, const float AnimationAlpha, const float GizmoScale, const float GizmoHoverScale, 
+	UE_API void UpdateHandlesRelativeTransformOnAxis( const FTransform& HandleToCenter, const float AnimationAlpha, const float GizmoScale, const float GizmoHoverScale, 
 		const FVector& ViewLocation, class UActorComponent* DraggingHandle, const TArray< UActorComponent* >& HoveringOverHandles );
 };
+
+#undef UE_API
diff --git a/Engine/Source/Editor/ViewportInteraction/Public/VISettings.h b/Engine/Source/Editor/ViewportInteraction/Public/VISettings.h
--- a/Engine/Source/Editor/ViewportInteraction/Public/VISettings.h
+++ b/Engine/Source/Editor/ViewportInteraction/Public/VISettings.h
@@ -10,8 +10,8 @@
 /**
 * Implements the settings for Viewport Interaction and VR Mode. Note that this is primarily a base class for anything that needs to be accessed in the VI module.
 */
-UCLASS(config = EditorSettings)
-class VIEWPORTINTERACTION_API UVISettings : public UObject
+UCLASS(MinimalAPI, config = EditorSettings)
+class UVISettings : public UObject
 {
 	GENERATED_BODY()
 
diff --git a/Engine/Source/Editor/ViewportInteraction/Public/ViewportDragOperation.h b/Engine/Source/Editor/ViewportInteraction/Public/ViewportDragOperation.h
--- a/Engine/Source/Editor/ViewportInteraction/Public/ViewportDragOperation.h
+++ b/Engine/Source/Editor/ViewportInteraction/Public/ViewportDragOperation.h
@@ -20,11 +20,13 @@
 
 #include "ViewportDragOperation.generated.h"
 
+#define UE_API VIEWPORTINTERACTION_API
+
 /**
  * Base class for interactable drag calculations
  */
-UCLASS()
-class VIEWPORTINTERACTION_API UViewportDragOperation : public UObject
+UCLASS(MinimalAPI)
+class UViewportDragOperation : public UObject
 {
 	GENERATED_BODY()
 
@@ -54,30 +56,30 @@ public:
 /**
  * Container component for UViewportDragOperation that can be used by objects in the world that are draggable and implement the ViewportInteractableInterface
  */
-UCLASS()
-class VIEWPORTINTERACTION_API UViewportDragOperationComponent : public UActorComponent
+UCLASS(MinimalAPI)
+class UViewportDragOperationComponent : public UActorComponent
 {
 	GENERATED_BODY()
 
 public:
 	
 	/** Default constructor. */
-	UViewportDragOperationComponent();
+	UE_API UViewportDragOperationComponent();
 
 	/** Get the actual dragging operation */
-	UViewportDragOperation* GetDragOperation();
+	UE_API UViewportDragOperation* GetDragOperation();
 
 	/** Sets the drag operation class that will be used next time starting dragging */
-	void SetDragOperationClass( const TSubclassOf<UViewportDragOperation> InDragOperation );
+	UE_API void SetDragOperationClass( const TSubclassOf<UViewportDragOperation> InDragOperation );
 
 	/** Starts new dragging operation */
-	void StartDragOperation();
+	UE_API void StartDragOperation();
 
 	/* Destroys the dragoperation */
-	void ClearDragOperation();
+	UE_API void ClearDragOperation();
 
 	/** If this operation is currently active. */
-	bool IsDragging() const;
+	UE_API bool IsDragging() const;
 
 private:
 
@@ -94,7 +96,7 @@ private:
  * Data structure that holds all arguments that can be used while dragging a transformable.
  */
 USTRUCT()
-struct VIEWPORTINTERACTION_API FDraggingTransformableData
+struct FDraggingTransformableData
 {
 	GENERATED_BODY()
 
@@ -154,3 +156,5 @@ struct VIEWPORTINTERACTION_API FDraggingTransformableData
 	bool bOutScaled;
 	bool bAllowSnap;
 };
+
+#undef UE_API
diff --git a/Engine/Source/Editor/ViewportInteraction/Public/ViewportInteractableInterface.h b/Engine/Source/Editor/ViewportInteraction/Public/ViewportInteractableInterface.h
--- a/Engine/Source/Editor/ViewportInteraction/Public/ViewportInteractableInterface.h
+++ b/Engine/Source/Editor/ViewportInteraction/Public/ViewportInteractableInterface.h
@@ -25,7 +25,7 @@ class UViewportInteractableInterface : public UInterface
 /**
  * Interface for custom objects that can be interacted with by a interactor
  */
-class VIEWPORTINTERACTION_API IViewportInteractableInterface
+class IViewportInteractableInterface
 {
 	GENERATED_IINTERFACE_BODY()
 
diff --git a/Engine/Source/Editor/ViewportInteraction/Public/ViewportInteractionAssetContainer.h b/Engine/Source/Editor/ViewportInteraction/Public/ViewportInteractionAssetContainer.h
--- a/Engine/Source/Editor/ViewportInteraction/Public/ViewportInteractionAssetContainer.h
+++ b/Engine/Source/Editor/ViewportInteraction/Public/ViewportInteractionAssetContainer.h
@@ -13,8 +13,8 @@ class UStaticMesh;
 /**
  * Asset container for viewport interaction.
  */
-UCLASS()
-class VIEWPORTINTERACTION_API UViewportInteractionAssetContainer : public UDataAsset
+UCLASS(MinimalAPI)
+class UViewportInteractionAssetContainer : public UDataAsset
 {
 	GENERATED_BODY()
 
diff --git a/Engine/Source/Editor/ViewportInteraction/Public/ViewportInteractionDragOperations.h b/Engine/Source/Editor/ViewportInteraction/Public/ViewportInteractionDragOperations.h
--- a/Engine/Source/Editor/ViewportInteraction/Public/ViewportInteractionDragOperations.h
+++ b/Engine/Source/Editor/ViewportInteraction/Public/ViewportInteractionDragOperations.h
@@ -8,46 +8,48 @@
 #include "ViewportDragOperation.h"	
 #include "ViewportInteractionDragOperations.generated.h"
 
+#define UE_API VIEWPORTINTERACTION_API
+
 /**
  * Gizmo translation on one axis.
  */
-UCLASS()
-class VIEWPORTINTERACTION_API UTranslationDragOperation : public UViewportDragOperation
+UCLASS(MinimalAPI)
+class UTranslationDragOperation : public UViewportDragOperation
 {
 	GENERATED_BODY()
 
 public:
-	virtual void ExecuteDrag(struct FDraggingTransformableData& DraggingData) override;
+	UE_API virtual void ExecuteDrag(struct FDraggingTransformableData& DraggingData) override;
 };
 
 /**
  * Gizmo translation on two axes.
  */
-UCLASS()
-class VIEWPORTINTERACTION_API UPlaneTranslationDragOperation: public UViewportDragOperation
+UCLASS(MinimalAPI)
+class UPlaneTranslationDragOperation: public UViewportDragOperation
 {
 	GENERATED_BODY()
 
 public:
-	UPlaneTranslationDragOperation();
-	virtual void ExecuteDrag(struct FDraggingTransformableData& DraggingData) override;
+	UE_API UPlaneTranslationDragOperation();
+	UE_API virtual void ExecuteDrag(struct FDraggingTransformableData& DraggingData) override;
 };
 
 /**
  * Rotation around one axis based on input angle.
  */
-UCLASS()
-class VIEWPORTINTERACTION_API URotateOnAngleDragOperation: public UViewportDragOperation
+UCLASS(MinimalAPI)
+class URotateOnAngleDragOperation: public UViewportDragOperation
 {
 	GENERATED_BODY()
 
 public:
-	URotateOnAngleDragOperation();
+	UE_API URotateOnAngleDragOperation();
 	
-	virtual void ExecuteDrag(struct FDraggingTransformableData& DraggingData) override;
+	UE_API virtual void ExecuteDrag(struct FDraggingTransformableData& DraggingData) override;
 
 	/** When RotateOnAngle we intersect on a plane to rotate the transform gizmo. This is the local point from the transform gizmo location of that intersect */
-	FVector GetLocalIntersectPointOnRotationGizmo() const;
+	UE_API FVector GetLocalIntersectPointOnRotationGizmo() const;
 
 private:
 
@@ -64,23 +66,25 @@ private:
 /**
  * Scale on one axis.
  */
-UCLASS()
-class VIEWPORTINTERACTION_API UScaleDragOperation : public UViewportDragOperation
+UCLASS(MinimalAPI)
+class UScaleDragOperation : public UViewportDragOperation
 {
 	GENERATED_BODY()
 
 public:
-	virtual void ExecuteDrag(struct FDraggingTransformableData& DraggingData) override;
+	UE_API virtual void ExecuteDrag(struct FDraggingTransformableData& DraggingData) override;
 };
 
 /**
  * Scale on all axes.
  */
-UCLASS()
-class  VIEWPORTINTERACTION_API UUniformScaleDragOperation : public UViewportDragOperation
+UCLASS(MinimalAPI)
+class  UUniformScaleDragOperation : public UViewportDragOperation
 {
 	GENERATED_BODY()
 
 public:
-	virtual void ExecuteDrag(struct FDraggingTransformableData& DraggingData) override;
+	UE_API virtual void ExecuteDrag(struct FDraggingTransformableData& DraggingData) override;
 };
+
+#undef UE_API
diff --git a/Engine/Source/Editor/ViewportInteraction/Public/ViewportInteractionTypes.h b/Engine/Source/Editor/ViewportInteraction/Public/ViewportInteractionTypes.h
--- a/Engine/Source/Editor/ViewportInteraction/Public/ViewportInteractionTypes.h
+++ b/Engine/Source/Editor/ViewportInteraction/Public/ViewportInteractionTypes.h
@@ -7,11 +7,13 @@
 #include "Engine/EngineBaseTypes.h"
 #include "ViewportInteractionTypes.generated.h"
 
+#define UE_API VIEWPORTINTERACTION_API
+
 /**
  * Represents a generic action
 */
 USTRUCT( BlueprintType )
-struct VIEWPORTINTERACTION_API FViewportActionKeyInput
+struct FViewportActionKeyInput
 {
 	GENERATED_BODY()
 	
@@ -111,7 +113,7 @@ enum class ETransformGizmoHandleDirection
 
 /** Placement of a handle in pivot space */
 USTRUCT()
-struct VIEWPORTINTERACTION_API FTransformGizmoHandlePlacement
+struct FTransformGizmoHandlePlacement
 {
 	GENERATED_BODY()
 
@@ -124,7 +126,7 @@ struct VIEWPORTINTERACTION_API FTransformGizmoHandlePlacement
 	the index (0-2) of the axis where the handle would be facing, or INDEX_NONE for corners or edges.
 	The center axis index is valid for edges, and defines the axis perpendicular to that edge direction,
 	or INDEX_NONE if it's not an edge */
-	void GetCenterHandleCountAndFacingAxisIndex( int32& OutCenterHandleCount, int32& OutFacingAxisIndex, int32& OutCenterAxisIndex ) const;
+	UE_API void GetCenterHandleCountAndFacingAxisIndex( int32& OutCenterHandleCount, int32& OutFacingAxisIndex, int32& OutCenterAxisIndex ) const;
 };
 
 
@@ -134,3 +136,5 @@ enum class ELockedWorldDragMode
 	OnlyRotating,
 	OnlyScaling,
 };
+
+#undef UE_API
diff --git a/Engine/Source/Editor/ViewportInteraction/Public/ViewportInteractionUtils.h b/Engine/Source/Editor/ViewportInteraction/Public/ViewportInteractionUtils.h
--- a/Engine/Source/Editor/ViewportInteraction/Public/ViewportInteractionUtils.h
+++ b/Engine/Source/Editor/ViewportInteraction/Public/ViewportInteractionUtils.h
@@ -6,13 +6,15 @@
 #include "Math/MathFwd.h"
 #include "Math/Vector.h"
 
+#define UE_API VIEWPORTINTERACTION_API
+
 namespace ViewportInteractionUtils
 {
 	/************************************************************************/
 	/* 1 Euro filter smoothing algorithm									*/
 	/* http://cristal.univ-lille.fr/~casiez/1euro/							*/
 	/************************************************************************/
-	class VIEWPORTINTERACTION_API FOneEuroFilter
+	class FOneEuroFilter
 	{
 	private:
 
@@ -44,27 +46,27 @@ namespace ViewportInteractionUtils
 	public:
 
 		/** Default constructor */
-		FOneEuroFilter();
+		UE_API FOneEuroFilter();
 
-		FOneEuroFilter(const double InMinCutoff, const double InCutoffSlope, const double InDeltaCutoff);
+		UE_API FOneEuroFilter(const double InMinCutoff, const double InCutoffSlope, const double InDeltaCutoff);
 
 		/** Smooth vector */
-		FVector Filter(const FVector& InRaw, const double InDeltaTime);
+		UE_API FVector Filter(const FVector& InRaw, const double InDeltaTime);
 
 		/** Set the minimum cutoff */
-		void SetMinCutoff(const double InMinCutoff);
+		UE_API void SetMinCutoff(const double InMinCutoff);
 
 		/** Set the cutoff slope */
-		void SetCutoffSlope(const double InCutoffSlope);
+		UE_API void SetCutoffSlope(const double InCutoffSlope);
 
 		/** Set the delta slope */
-		void SetDeltaCutoff(const double InDeltaCutoff);
+		UE_API void SetDeltaCutoff(const double InDeltaCutoff);
 
 	private:
 
-		const FVector CalculateCutoff(const FVector& InValue);
-		const FVector CalculateAlpha(const FVector& InCutoff, const double InDeltaTime) const;
-		const double CalculateAlpha(const double InCutoff, const double InDeltaTime) const;
+		UE_API const FVector CalculateCutoff(const FVector& InValue);
+		UE_API const FVector CalculateAlpha(const FVector& InCutoff, const double InDeltaTime) const;
+		UE_API const double CalculateAlpha(const double InCutoff, const double InDeltaTime) const;
 
 		double MinCutoff;
 		double CutoffSlope;
@@ -72,4 +74,6 @@ namespace ViewportInteractionUtils
 		FLowpassFilter RawFilter;
 		FLowpassFilter DeltaFilter;
 	};
-}
\ No newline at end of file
+}
+
+#undef UE_API
diff --git a/Engine/Source/Editor/ViewportInteraction/Public/ViewportInteractor.h b/Engine/Source/Editor/ViewportInteraction/Public/ViewportInteractor.h
--- a/Engine/Source/Editor/ViewportInteraction/Public/ViewportInteractor.h
+++ b/Engine/Source/Editor/ViewportInteraction/Public/ViewportInteractor.h
@@ -12,6 +12,8 @@
 #include "ViewportInteractionUtils.h"
 #include "ViewportInteractor.generated.h"
 
+#define UE_API VIEWPORTINTERACTION_API
+
 class AActor;
 
 /** Filter mode for GetHitResultFromLaserPointer*/
@@ -26,55 +28,55 @@ enum class EHitResultGizmoFilterMode : uint8
 /**
  * Represents the interactor in the world
  */
-UCLASS( Abstract, Blueprintable, BlueprintType)
-class VIEWPORTINTERACTION_API UViewportInteractor : public UObject
+UCLASS(MinimalAPI,  Abstract, Blueprintable, BlueprintType)
+class UViewportInteractor : public UObject
 {
 	GENERATED_BODY()
 
 public:
 
 	/** Default constructor */
-	UViewportInteractor();
+	UE_API UViewportInteractor();
 
 	/** Gets the private data for this interactor */
-	struct FViewportInteractorData& GetInteractorData();
+	UE_API struct FViewportInteractorData& GetInteractorData();
 
 	/** Gets the private data for this interactor (const) */
-	const struct FViewportInteractorData& GetInteractorData() const;
+	UE_API const struct FViewportInteractorData& GetInteractorData() const;
 
 	/** Sets the world interaction */
-	void SetWorldInteraction( class UViewportWorldInteraction* InWorldInteraction );
+	UE_API void SetWorldInteraction( class UViewportWorldInteraction* InWorldInteraction );
 
 	/** Gets the world interaction */
 	UFUNCTION(BlueprintCallable, Category = "Interactor")
-	UViewportWorldInteraction* GetWorldInteraction() const;
+	UE_API UViewportWorldInteraction* GetWorldInteraction() const;
 
 	/** Sets the other interactor */
-	void SetOtherInteractor( UViewportInteractor* InOtherInteractor ); //@todo ViewportInteraction: This should not be public to other modules and should only be called from the world interaction.
+	UE_API void SetOtherInteractor( UViewportInteractor* InOtherInteractor ); //@todo ViewportInteraction: This should not be public to other modules and should only be called from the world interaction.
 
 	/** Removes the other interactor reference for this interactor */
-	void RemoveOtherInteractor();
+	UE_API void RemoveOtherInteractor();
 
 	/** Gets the paired interactor assigned by the world interaction, can return null when there is no other interactor */
 	UFUNCTION(BlueprintCallable, Category = "Interactor")
-	UViewportInteractor* GetOtherInteractor() const;
+	UE_API UViewportInteractor* GetOtherInteractor() const;
 
 	/** Whenever the ViewportWorldInteraction is shut down, the interacts will shut down as well */
 	UFUNCTION(BlueprintNativeEvent, CallInEditor, Category = "Interactor")
-	void Shutdown();
+	UE_API void Shutdown();
 
 	/** Update for this interactor called by the ViewportWorldInteraction */
 	UFUNCTION(BlueprintNativeEvent, CallInEditor, Category = "Interactor")
-	void Tick( const float DeltaTime );
+	UE_API void Tick( const float DeltaTime );
 
 	/** Gets the last component hovered on by the interactor laser. */
-	virtual class UActorComponent* GetLastHoverComponent();
+	UE_API virtual class UActorComponent* GetLastHoverComponent();
 
 	/** Adds a new action to the KeyToActionMap */
-	void AddKeyAction( const FKey& Key, const FViewportActionKeyInput& Action );
+	UE_API void AddKeyAction( const FKey& Key, const FViewportActionKeyInput& Action );
 
 	/** Removes an action from the KeyToActionMap */
-	void RemoveKeyAction( const FKey& Key );
+	UE_API void RemoveKeyAction( const FKey& Key );
 
 	/** Base classes need to implement getting the input for the input devices for that interactor */
 	virtual void PollInput() {}
@@ -85,7 +87,7 @@ public:
 	 *  BP will have there "Receive" versions called from within these function and do not have to call
 	 *  there parents.  They simply need to return if they handled the input or not.
 	 */
-	bool HandleInputKey( class FEditorViewportClient& ViewportClient, const FKey Key, const EInputEvent Event );
+	UE_API bool HandleInputKey( class FEditorViewportClient& ViewportClient, const FKey Key, const EInputEvent Event );
 
 	/**
 	 * Handles axis input and translates it to actions
@@ -93,7 +95,7 @@ public:
 	 *  BP will have there "Receive" versions called from within these function and do not have to call
 	 *  there parents.They simply need to return if they handled the input or not.
 	 */
-	bool HandleInputAxis( class FEditorViewportClient& ViewportClient, const FKey Key, const float Delta, const float DeltaTime );
+	UE_API bool HandleInputAxis( class FEditorViewportClient& ViewportClient, const FKey Key, const float Delta, const float DeltaTime );
 
 	/** @return	Returns true if this interactor's designated 'modifier button' is currently held down.  Some interactors may not support this */
 	virtual bool IsModifierPressed() const
@@ -103,32 +105,32 @@ public:
 
 	/** Gets the world transform of this interactor */
 	UFUNCTION(BlueprintCallable, Category = "Interactor")
-	FTransform GetTransform() const;
+	UE_API FTransform GetTransform() const;
 
 	/** Gets the hand transform of the interactor, in the local tracking space */
 	UFUNCTION(BlueprintCallable, Category = "Interactor")
-	FTransform GetRoomSpaceTransform() const;
+	UE_API FTransform GetRoomSpaceTransform() const;
 
 	/** Gets the last world transform of this interactor */
 	UFUNCTION(BlueprintCallable, Category = "Interactor")
-	FTransform GetLastTransform() const;
+	UE_API FTransform GetLastTransform() const;
 
 	/** Gets the last hand transform of the interactor, in the local tracking space */
 	UFUNCTION(BlueprintCallable, Category = "Interactor")
-	FTransform GetLastRoomSpaceTransform() const;
+	UE_API FTransform GetLastRoomSpaceTransform() const;
 
 	/** Gets the current interactor data dragging mode */
 	UFUNCTION(BlueprintCallable, Category = "Interactor")
-	EViewportInteractionDraggingMode GetDraggingMode() const;
+	UE_API EViewportInteractionDraggingMode GetDraggingMode() const;
 
 	/** Gets the interactor data previous dragging mode */
-	EViewportInteractionDraggingMode GetLastDraggingMode() const;
+	UE_API EViewportInteractionDraggingMode GetLastDraggingMode() const;
 
 	/** Gets the interactor data drag velocity */
-	FVector GetDragTranslationVelocity() const;
+	UE_API FVector GetDragTranslationVelocity() const;
 
 	/** Sets the hover location */
-	void SetHoverLocation(const FVector& InHoverLocation);
+	UE_API void SetHoverLocation(const FVector& InHoverLocation);
 
 	/**
 	 * Gets the start and end point of the laser pointer for the specified hand
@@ -141,7 +143,7 @@ public:
 	 * @return	True if we have motion controller data for this hand and could return a valid result
 	 */
 	UFUNCTION(BlueprintCallable, Category = "Interactor")
-	bool GetLaserPointer( FVector& LaserPointerStart, FVector& LaserPointerEnd, const bool bEvenIfBlocked = false, const float LaserLengthOverride = 0.0f );
+	UE_API bool GetLaserPointer( FVector& LaserPointerStart, FVector& LaserPointerEnd, const bool bEvenIfBlocked = false, const float LaserLengthOverride = 0.0f );
 
 	/**
 	 * Gets a sphere on this interactor that can be used to interact with objects in close proximity
@@ -151,10 +153,10 @@ public:
 	 *
 	 * @return	True if the sphere is available and valid, or false is the interactor was busy or we could not determine a valid position for the interactor
 	 */
-	bool GetGrabberSphere( FSphere& OutGrabberSphere, const bool bEvenIfBlocked = false );
+	UE_API bool GetGrabberSphere( FSphere& OutGrabberSphere, const bool bEvenIfBlocked = false );
 
 	/** Gets the maximum length of a laser pointer */
-	float GetLaserPointerMaxLength() const;
+	UE_API float GetLaserPointerMaxLength() const;
 
 	/** Triggers a force feedback effect on device if possible */
 	virtual void PlayHapticEffect( const float Strength ) {};
@@ -186,10 +188,10 @@ public:
 	 *
 	 * @return What the laster pointer hit
 	 */
-	virtual FHitResult GetHitResultFromLaserPointer( TArray<AActor*>* OptionalListOfIgnoredActors = nullptr, const EHitResultGizmoFilterMode GizmoFilterMode = EHitResultGizmoFilterMode::All, TArray<UClass*>* ObjectsInFrontOfGizmo = nullptr, const bool bEvenIfBlocked = false, const float LaserLengthOverride = 0.0f );
+	UE_API virtual FHitResult GetHitResultFromLaserPointer( TArray<AActor*>* OptionalListOfIgnoredActors = nullptr, const EHitResultGizmoFilterMode GizmoFilterMode = EHitResultGizmoFilterMode::All, TArray<UClass*>* ObjectsInFrontOfGizmo = nullptr, const bool bEvenIfBlocked = false, const float LaserLengthOverride = 0.0f );
 
 	/** Reset the values before checking the hover actions */
-	virtual void ResetHoverState();
+	UE_API virtual void ResetHoverState();
 
 	/** Needs to be implemented by the base to calculate drag ray length and the velocity for the ray */
 	virtual void CalculateDragRay( double& InOutDragRayLength, double& InOutDragRayVelocity ) {};
@@ -206,60 +208,60 @@ public:
 	 * @return	True if we have motion controller data for this hand and could return a valid result
 	 */
 	UFUNCTION(BlueprintCallable, Category = "Interactor")
-	virtual bool GetTransformAndForwardVector( FTransform& OutHandTransform, FVector& OutForwardVector ) const;
+	UE_API virtual bool GetTransformAndForwardVector( FTransform& OutHandTransform, FVector& OutForwardVector ) const;
 
 	/** Called by StartDragging in world interaction to give the interactor a chance of acting upon starting a drag operation */
-	virtual void OnStartDragging( const FVector& HitLocation, const bool bIsPlacingNewObjects );
+	UE_API virtual void OnStartDragging( const FVector& HitLocation, const bool bIsPlacingNewObjects );
 
 	/** Gets the interactor laser hover location */
 	UFUNCTION(BlueprintCallable, Category = "Interactor")
-	FVector GetHoverLocation();
+	UE_API FVector GetHoverLocation();
 
 	/** If the interactor laser is currently hovering */
-	bool IsHovering() const;
+	UE_API bool IsHovering() const;
 
 	/** If the interactor laser is currently hovering over a gizmo handle */
 	UFUNCTION(BlueprintCallable, Category = "Interactor")
-	bool IsHoveringOverGizmo() const;
+	UE_API bool IsHoveringOverGizmo() const;
 
 	/** Sets the current dragging mode for this interactor */
 	UFUNCTION(BlueprintCallable, Category = "Interactor")
-	void SetDraggingMode( const EViewportInteractionDraggingMode NewDraggingMode );
+	UE_API void SetDraggingMode( const EViewportInteractionDraggingMode NewDraggingMode );
 
 	/** To be overridden by base class. Called by GetLaserPointer to give the derived interactor a chance to disable the laser. By default it is not blocked */
 	virtual bool GetIsLaserBlocked() const { return false; }
 
 	/** Gets a certain action by iterating through the map looking for the same ActionType */
 	// @todo ViewportInteractor: This should be changed to return a const pointer, but we need to fix up some dragging code in WorldInteraction first
-	FViewportActionKeyInput* GetActionWithName( const FName InActionName );
+	UE_API FViewportActionKeyInput* GetActionWithName( const FName InActionName );
 
 	/** Gets the drag haptic feedback strength console variable */
-	float GetDragHapticFeedbackStrength() const;
+	UE_API float GetDragHapticFeedbackStrength() const;
 
 	/** If this interactor is hovering over a type that has priority from GetHitResultFromLaserPointer */
-	bool IsHoveringOverPriorityType() const;
+	UE_API bool IsHoveringOverPriorityType() const;
 
 	/** Returns true if currently hovering over selected actor */
-	bool IsHoveringOverSelectedActor() const;
+	UE_API bool IsHoveringOverSelectedActor() const;
 
 	/** Reset the stored laser end location at the end of tick */
-	void ResetLaserEnd();
+	UE_API void ResetLaserEnd();
 
 	/** Sets the current gizmo filter mode used for Interaction*/
 	UFUNCTION(BlueprintCallable, Category = "Interactor")
-	void SetHitResultGizmoFilterMode(EHitResultGizmoFilterMode newFilter);
+	UE_API void SetHitResultGizmoFilterMode(EHitResultGizmoFilterMode newFilter);
 
 	/** Gets current gizmo filter mode used for Interaction*/
 	UFUNCTION(BlueprintCallable, Category = "Interactor")
-	EHitResultGizmoFilterMode GetHitResultGizmoFilterMode() const;
+	UE_API EHitResultGizmoFilterMode GetHitResultGizmoFilterMode() const;
 
 	/** Sets if the interactor can carry an object */
 	UFUNCTION(BlueprintCallable, Category = "Interactor")
-	void SetCanCarry(const bool bInCanCarry);
+	UE_API void SetCanCarry(const bool bInCanCarry);
 
 	/** Gets if the interactor can carry an object */
 	UFUNCTION(BlueprintCallable, Category = "Interactor")
-	bool CanCarry() const;
+	UE_API bool CanCarry() const;
 
 protected:
 
@@ -274,17 +276,17 @@ protected:
 
 	/** To be overridden. Called by HandleInputKey before delegates and default input implementation */
 	UFUNCTION(BlueprintImplementableEvent, CallInEditor, meta = (DisplayName = "HandleInputKey"), Category = "Interactor")
-	void HandleInputKey_BP( const FViewportActionKeyInput& Action, const FKey Key, const EInputEvent Event, bool& bOutWasHandled);
+	UE_API void HandleInputKey_BP( const FViewportActionKeyInput& Action, const FKey Key, const EInputEvent Event, bool& bOutWasHandled);
 
 	/** To be overridden. Called by HandleInputAxis before delegates and default input implementation */
 	virtual void HandleInputAxis( class FEditorViewportClient& ViewportClient, FViewportActionKeyInput& Action, const FKey Key, const float Delta, const float DeltaTime, bool& bOutWasHandled ) {};
 
 	/** To be overridden. Called by HandleInputAxis before delegates and default input implementation */
 	UFUNCTION(BlueprintImplementableEvent, CallInEditor, meta = (DisplayName = "HandleInputAxis"), Category = "Interactor")
-	void HandleInputAxis_BP( const FViewportActionKeyInput& Action, const FKey Key, const float Delta, const float DeltaTime, bool& bOutWasHandled);
+	UE_API void HandleInputAxis_BP( const FViewportActionKeyInput& Action, const FKey Key, const float Delta, const float DeltaTime, bool& bOutWasHandled);
 
 	/** If this interactors allows to smooth the laser. Default is true. */
-	virtual bool AllowLaserSmoothing() const;
+	UE_API virtual bool AllowLaserSmoothing() const;
 
 protected:
 
@@ -324,4 +326,6 @@ private:
 
 	/** Smoothing filter for laser */
 	ViewportInteractionUtils::FOneEuroFilter SmoothingOneEuroFilter;
-};
\ No newline at end of file
+};
+
+#undef UE_API
diff --git a/Engine/Source/Editor/ViewportInteraction/Public/ViewportTransformable.h b/Engine/Source/Editor/ViewportInteraction/Public/ViewportTransformable.h
--- a/Engine/Source/Editor/ViewportInteraction/Public/ViewportTransformable.h
+++ b/Engine/Source/Editor/ViewportInteraction/Public/ViewportTransformable.h
@@ -12,7 +12,7 @@
 /**
  * Represents an object that we're actively interacting with, such as a selected actor
  */
-class VIEWPORTINTERACTION_API FViewportTransformable
+class FViewportTransformable
 {
 
 public:
diff --git a/Engine/Source/Editor/ViewportInteraction/Public/ViewportTransformer.h b/Engine/Source/Editor/ViewportInteraction/Public/ViewportTransformer.h
--- a/Engine/Source/Editor/ViewportInteraction/Public/ViewportTransformer.h
+++ b/Engine/Source/Editor/ViewportInteraction/Public/ViewportTransformer.h
@@ -9,18 +9,20 @@
 #include "ViewportWorldInteraction.h"
 #include "ViewportTransformer.generated.h"
 
-UCLASS( abstract )
-class VIEWPORTINTERACTION_API UViewportTransformer : public UObject
+#define UE_API VIEWPORTINTERACTION_API
+
+UCLASS(MinimalAPI,  abstract )
+class UViewportTransformer : public UObject
 {
 	GENERATED_BODY()
 
 public:
 
 	UFUNCTION()
-	virtual void Init( UViewportWorldInteraction* InitViewportWorldInteraction );
+	UE_API virtual void Init( UViewportWorldInteraction* InitViewportWorldInteraction );
 
 	UFUNCTION()
-	virtual void Shutdown();
+	UE_API virtual void Shutdown();
 
 	/** @return If this transformer can be used to align its transformable to actors in the scene */
 	UFUNCTION()
@@ -52,3 +54,4 @@ protected:
 
 };
 
+#undef UE_API
diff --git a/Engine/Source/Editor/ViewportInteraction/Public/ViewportWorldInteraction.h b/Engine/Source/Editor/ViewportInteraction/Public/ViewportWorldInteraction.h
--- a/Engine/Source/Editor/ViewportInteraction/Public/ViewportWorldInteraction.h
+++ b/Engine/Source/Editor/ViewportInteraction/Public/ViewportWorldInteraction.h
@@ -19,6 +19,8 @@
 #include "GenericPlatform/ICursor.h"
 #include "ViewportWorldInteraction.generated.h"
 
+#define UE_API VIEWPORTINTERACTION_API
+
 namespace ViewportWorldActionTypes
 {
 	static const FName NoAction( "NoAction" );
@@ -43,37 +45,37 @@ enum class EViewportWorldInteractionType : uint8
 	Legacy = 1
 };
 
-UCLASS(BlueprintType, Transient)
-class VIEWPORTINTERACTION_API UViewportWorldInteraction : public UEditorWorldExtension
+UCLASS(MinimalAPI, BlueprintType, Transient)
+class UViewportWorldInteraction : public UEditorWorldExtension
 {
 	GENERATED_BODY()
 
 public:
 
-	UViewportWorldInteraction();
+	UE_API UViewportWorldInteraction();
 	// UEditorExtension overrides
-	virtual void Init() override;
-	virtual void Shutdown() override;
-	virtual void Tick( float DeltaSeconds ) override;
+	UE_API virtual void Init() override;
+	UE_API virtual void Shutdown() override;
+	UE_API virtual void Tick( float DeltaSeconds ) override;
 
 	/** Initialize colors */
-	void InitColors();
+	UE_API void InitColors();
 
 	/** Adds interactor to the worldinteraction */
 	UFUNCTION(BlueprintCallable, Category = "ViewportWorldInteraction")
-	void AddInteractor( UViewportInteractor* Interactor );
+	UE_API void AddInteractor( UViewportInteractor* Interactor );
 
 	/** Removes interactor from the worldinteraction and removes the interactor from its paired interactor if any */
 	UFUNCTION(BlueprintCallable, Category = "ViewportWorldInteraction")
-	void RemoveInteractor( UViewportInteractor* Interactor );
+	UE_API void RemoveInteractor( UViewportInteractor* Interactor );
 
 	/** Creates an interactor for the mouse cursor.  If one already exists, this will add a reference to it.  Remember to
 	    call ReleaseMouseCursorInteractor() when you're done. */
-	void AddMouseCursorInteractor();
+	UE_API void AddMouseCursorInteractor();
 
 	/** When you're finished with the mouse cursor interactor, call this to release it.  It will be destroyed unless
 	    another system is still using it */
-	void ReleaseMouseCursorInteractor();
+	UE_API void ReleaseMouseCursorInteractor();
 
 	/** Gets the event for hovering update which is broadcasted for each interactor */
 	DECLARE_EVENT_ThreeParams( UViewportWorldInteraction, FOnVIHoverUpdate, class UViewportInteractor* /* Interactor */, FVector& /* OutHoverImpactPoint */, bool& /* bWasHandled */ );
@@ -117,7 +119,7 @@ public:
 
 	/** Sets the system that should be used to transform objects in the scene.  Only one can be active at a time.  If you pass
 	    nullptr, a default transformer for actors/components will be enabled */
-	void SetTransformer( class UViewportTransformer* NewTransformer );
+	UE_API void SetTransformer( class UViewportTransformer* NewTransformer );
 
 	/** @return	Gets the currently active transformer object */
 	const UViewportTransformer* GetTransformer() const
@@ -127,13 +129,13 @@ public:
 
 	/** Sets the list of objects that this system will be responsible for transforming when interacting using the
 	    gizmo or directly on the objects */
-	void SetTransformables( TArray< TUniquePtr< FViewportTransformable > >&& NewTransformables, const bool bNewObjectsSelected );
+	UE_API void SetTransformables( TArray< TUniquePtr< FViewportTransformable > >&& NewTransformables, const bool bNewObjectsSelected );
 
 	/** When using VR, this sets the viewport client that's been "possessed" by the head mounted display.  Only valid when VR is enabled. */
-	void SetDefaultOptionalViewportClient( const TSharedPtr<class FEditorViewportClient>& InEditorViewportClient );
+	UE_API void SetDefaultOptionalViewportClient( const TSharedPtr<class FEditorViewportClient>& InEditorViewportClient );
 
 	/** Pairs to interactors by setting	the other interactor for each interactor */
-	void PairInteractors( UViewportInteractor* FirstInteractor, UViewportInteractor* SecondInteractor );
+	UE_API void PairInteractors( UViewportInteractor* FirstInteractor, UViewportInteractor* SecondInteractor );
 
 	/**
 	 * Adds an actor to the list of actors to never allow an interactor to hit in the scene.  No selection.  No hover.
@@ -142,7 +144,7 @@ public:
 	 * @param	ActorToExcludeFromHitTests	The actor that should be forever excluded from hit tests
 	 */
 	UFUNCTION(BlueprintCallable, Category = "ViewportWorldInteraction")
-	void AddActorToExcludeFromHitTests( AActor* ActorToExcludeFromHitTests );
+	UE_API void AddActorToExcludeFromHitTests( AActor* ActorToExcludeFromHitTests );
 
 	/** @return	Returns the list of actors that should never be included in interctor hit tests (hovering or selection) */
 	const TArray<TWeakObjectPtr<AActor>>& GetActorsToExcludeFromHitTest() const
@@ -154,48 +156,48 @@ public:
 	// Input
  	//
 
-	virtual bool InputKey( FEditorViewportClient* InViewportClient, FViewport* Viewport, FKey Key, EInputEvent Event ) override;
-	virtual bool InputAxis( FEditorViewportClient* InViewportClient, FViewport* Viewport, int32 ControllerId, FKey Key, float Delta, float DeltaTime ) override;
+	UE_API virtual bool InputKey( FEditorViewportClient* InViewportClient, FViewport* Viewport, FKey Key, EInputEvent Event ) override;
+	UE_API virtual bool InputAxis( FEditorViewportClient* InViewportClient, FViewport* Viewport, int32 ControllerId, FKey Key, float Delta, float DeltaTime ) override;
 
 	/** Handles the key input from the Input Preprocessor*/
-	bool PreprocessedInputKey( const FKey Key, const EInputEvent Event );
+	UE_API bool PreprocessedInputKey( const FKey Key, const EInputEvent Event );
 
 	/** Handles the axis input from the Input Preprocessor*/
-	bool PreprocessedInputAxis( const int32 ControllerId, const FKey Key, const float Delta, const double DeltaTime );
+	UE_API bool PreprocessedInputAxis( const int32 ControllerId, const FKey Key, const float Delta, const double DeltaTime );
 
 	/** Gets the world space transform of the calibrated VR room origin.  When using a seated VR device, this will feel like the
 	camera's world transform (before any HMD positional or rotation adjustments are applied.) */
 	UFUNCTION(BlueprintCallable, Category = "ViewportWorldInteraction")
-	FTransform GetRoomTransform() const;
+	UE_API FTransform GetRoomTransform() const;
 
 	/** Gets the transform of the viewport / user's HMD in room space */
 	UFUNCTION(BlueprintCallable, Category = "ViewportWorldInteraction")
-	FTransform GetRoomSpaceHeadTransform() const;
+	UE_API FTransform GetRoomSpaceHeadTransform() const;
 
 	/** Gets the transform of the viewport / user's HMD in world space */
 	UFUNCTION(BlueprintCallable, Category = "ViewportWorldInteraction")
-	FTransform GetHeadTransform() const;
+	UE_API FTransform GetHeadTransform() const;
 
 	/** Sets a new transform for the room so that the HMD is aligned to the new transform.
 		The Head is kept level to the ground and only rotated on the yaw */
 	UFUNCTION(BlueprintCallable, Category = "ViewportWorldInteraction")
-	void SetHeadTransform( const FTransform& NewHeadTransform );
+	UE_API void SetHeadTransform( const FTransform& NewHeadTransform );
 
 	/** Returns true if we actually are using VR and have a valid head location, meaning a call to GetHeadTransform() is valid */
-	bool HaveHeadTransform() const;
+	UE_API bool HaveHeadTransform() const;
 
 	/** Sets a new transform for the room, in world space.  This is basically setting the editor's camera transform for the viewport */
-	void SetRoomTransform( const FTransform& NewRoomTransform );
+	UE_API void SetRoomTransform( const FTransform& NewRoomTransform );
 
 	UFUNCTION(BlueprintCallable, Category = "ViewportWorldInteraction")
-	void SetRoomTransformForNextFrame(const FTransform& NewRoomTransform);
+	UE_API void SetRoomTransformForNextFrame(const FTransform& NewRoomTransform);
 
 	/** Gets the world scale factor, which can be multiplied by a scale vector to convert to room space */
 	UFUNCTION(BlueprintCallable, Category = "ViewportWorldInteraction")
-	float GetWorldScaleFactor() const;
+	UE_API float GetWorldScaleFactor() const;
 
 	/** Gets the currently used viewport */
-	FEditorViewportClient* GetDefaultOptionalViewportClient() const;
+	UE_API FEditorViewportClient* GetDefaultOptionalViewportClient() const;
 
 	/** Returns the Unreal controller ID for the motion controllers we're using */
 	int32 GetMotionControllerID() const
@@ -204,50 +206,50 @@ public:
 	}
 
 	/** Invokes the editor's undo system to undo the last change */
-	void Undo();
+	UE_API void Undo();
 
 	/** Redoes the last change */
-	void Redo();
+	UE_API void Redo();
 
 	/** Deletes all the selected objects */
-	void DeleteSelectedObjects();
+	UE_API void DeleteSelectedObjects();
 
 	/** Copies selected objects to the clipboard */
-	void Copy();
+	UE_API void Copy();
 
 	/** Pastes the clipboard contents into the scene */
-	void Paste();
+	UE_API void Paste();
 
 	/** Duplicates the selected objects */
-	void Duplicate();
+	UE_API void Duplicate();
 
 	/** Deselects all the current selected objects */
-	void Deselect();
+	UE_API void Deselect();
 
 	/** Called to finish a drag action with the specified interactor */
-	void StopDragging( class UViewportInteractor* Interactor );
+	UE_API void StopDragging( class UViewportInteractor* Interactor );
 
 	/** Starts dragging selected objects around.  Called when clicking and dragging on actors/gizmos in the world, or when placing new objects. */
-	void StartDragging( UViewportInteractor* Interactor, UActorComponent* ClickedTransformGizmoComponent, const FVector& HitLocation, const bool bIsPlacingNewObjects, const bool bAllowInterpolationWhenPlacing, const bool bShouldUseLaserImpactDrag, const bool bStartTransaction, const bool bWithGrabberSphere );
+	UE_API void StartDragging( UViewportInteractor* Interactor, UActorComponent* ClickedTransformGizmoComponent, const FVector& HitLocation, const bool bIsPlacingNewObjects, const bool bAllowInterpolationWhenPlacing, const bool bShouldUseLaserImpactDrag, const bool bStartTransaction, const bool bWithGrabberSphere );
 
 	DECLARE_EVENT_OneParam( UViewportWorldInteraction, FOnWorldScaleChanged, const float /* NewWorldToMetersScale */);
 	virtual FOnWorldScaleChanged& OnWorldScaleChanged() { return OnWorldScaleChangedEvent; };
 
 	/** Sets which transform gizmo coordinate space is used */
-	void SetTransformGizmoCoordinateSpace( const ECoordSystem NewCoordSystem );
+	UE_API void SetTransformGizmoCoordinateSpace( const ECoordSystem NewCoordSystem );
 
 	/** Returns which transform gizmo coordinate space we're using, world or local */
-	ECoordSystem GetTransformGizmoCoordinateSpace() const;
+	UE_API ECoordSystem GetTransformGizmoCoordinateSpace() const;
 
 	/** Returns the maximum user scale */
-	float GetMaxScale();
+	UE_API float GetMaxScale();
 
 	/** Returns the minimum user scale */
-	float GetMinScale();
+	UE_API float GetMinScale();
 
 	/** Sets GNewWorldToMetersScale */
 	UFUNCTION(BlueprintCallable, Category = "ViewportWorldInteraction")
-	void SetWorldToMetersScale( const float NewWorldToMetersScale, const bool bCompensateRoomWorldScale = false );
+	UE_API void SetWorldToMetersScale( const float NewWorldToMetersScale, const bool bCompensateRoomWorldScale = false );
 
 	/** Tells the world interaction system to skip updates of world movement this frame.  This is useful if you've called
 	    SetWorldToMetersScale() yourself to set a static world scale, and the cached room-space transforms for the last
@@ -258,67 +260,67 @@ public:
 	}
 
 	/** Checks to see if we're allowed to interact with the specified component */
-	bool IsInteractableComponent( const UActorComponent* Component ) const;
+	UE_API bool IsInteractableComponent( const UActorComponent* Component ) const;
 
 	/** Gets the transform gizmo actor, or returns null if we currently don't have one */
 	UFUNCTION(BlueprintCallable, Category = "ViewportWorldInteraction")
-	class ABaseTransformGizmo* GetTransformGizmoActor();
+	UE_API class ABaseTransformGizmo* GetTransformGizmoActor();
 
 	/** Sets whether the transform gizmo should be visible at all */
-	void SetTransformGizmoVisible( const bool bShouldBeVisible );
+	UE_API void SetTransformGizmoVisible( const bool bShouldBeVisible );
 
 	/** Returns whether the transform gizmo should be visible */
-	bool ShouldTransformGizmoBeVisible() const;
+	UE_API bool ShouldTransformGizmoBeVisible() const;
 
 	/** Returns whether the transform gizmo is actually visible right now.  This is different than whether it 'should' be
 	    visible, as the gizmo can be implicitly hidden if there are no selected objects, for example. */
-	bool IsTransformGizmoVisible() const;
+	UE_API bool IsTransformGizmoVisible() const;
 
 	/** Sets how large the transform gizmo is */
-	void SetTransformGizmoScale( const float NewScale );
+	UE_API void SetTransformGizmoScale( const float NewScale );
 
 	/** Returns the current transform gizmo scale */
-	float GetTransformGizmoScale() const;
+	UE_API float GetTransformGizmoScale() const;
 
 	/** Sets if objects are dragged with either hand since last time selecting something */
-	void SetDraggedSinceLastSelection( const bool bNewDraggedSinceLastSelection );
+	UE_API void SetDraggedSinceLastSelection( const bool bNewDraggedSinceLastSelection );
 
 	/** Sets the transform of the gizmo when starting drag */
-	void SetLastDragGizmoStartTransform( const FTransform NewLastDragGizmoStartTransform );
+	UE_API void SetLastDragGizmoStartTransform( const FTransform NewLastDragGizmoStartTransform );
 
 	/** Gets all the interactors */
 	UFUNCTION(BlueprintCallable, Category = "ViewportWorldInteraction")
-	const TArray<UViewportInteractor*>& GetInteractors() const;
+	UE_API const TArray<UViewportInteractor*>& GetInteractors() const;
 
 	/** Given a world space velocity vector, applies inertial damping to it to slow it down */
-	void ApplyVelocityDamping( FVector& Velocity, const bool bVelocitySensitive );
+	UE_API void ApplyVelocityDamping( FVector& Velocity, const bool bVelocitySensitive );
 
 	/** Gets the current Gizmo handle type */
-	EGizmoHandleTypes GetCurrentGizmoType() const;
+	UE_API EGizmoHandleTypes GetCurrentGizmoType() const;
 
 	/** Sets the current gizmo handle type */
-	void SetGizmoHandleType( const EGizmoHandleTypes InGizmoHandleType );
+	UE_API void SetGizmoHandleType( const EGizmoHandleTypes InGizmoHandleType );
 
 	/** Changes the transform gizmo class that it will change to next refresh */
-	void SetTransformGizmoClass( const TSubclassOf<ABaseTransformGizmo>& NewTransformGizmoClass );
+	UE_API void SetTransformGizmoClass( const TSubclassOf<ABaseTransformGizmo>& NewTransformGizmoClass );
 
 	/** Sets the currently dragged interactavle */
-	void SetDraggedInteractable( IViewportInteractableInterface* InDraggedInteractable, UViewportInteractor* Interactor );
+	UE_API void SetDraggedInteractable( IViewportInteractableInterface* InDraggedInteractable, UViewportInteractor* Interactor );
 
 	/** Check if there is another interactor hovering over the component */
-	bool IsOtherInteractorHoveringOverComponent( UViewportInteractor* Interactor, UActorComponent* Component ) const;
+	UE_API bool IsOtherInteractorHoveringOverComponent( UViewportInteractor* Interactor, UActorComponent* Component ) const;
 
 	/** Switch which transform gizmo coordinate space we're using. */
-	void CycleTransformGizmoCoordinateSpace();
+	UE_API void CycleTransformGizmoCoordinateSpace();
 
 	/** Figures out where to place an object when tracing it against the scene using a laser pointer */
-	bool FindPlacementPointUnderLaser( UViewportInteractor* Interactor, FVector& OutHitLocation );
+	UE_API bool FindPlacementPointUnderLaser( UViewportInteractor* Interactor, FVector& OutHitLocation );
 
 	/** Gets the tracking transactions */
-	FTrackingTransaction& GetTrackingTransaction();
+	UE_API FTrackingTransaction& GetTrackingTransaction();
 
 	/** Sets whether we should allow input events from the input preprocessor or not */
-	void SetUseInputPreprocessor( bool bInUseInputPreprocessor );
+	UE_API void SetUseInputPreprocessor( bool bInUseInputPreprocessor );
 
 	/** Gets a list of all of the objects we're currently interacting with, such as selected actors */
 	TArray< TUniquePtr< FViewportTransformable > >& GetTransformables()
@@ -327,72 +329,72 @@ public:
 	}
 
 	/** The ability to move and scale the world */
-	void AllowWorldMovement( bool bAllow );
+	UE_API void AllowWorldMovement( bool bAllow );
 
 	/** For other systems to check if the Viewport World Interaction system is currently aligning transformables to actors*/
-	bool AreAligningToActors() const;
+	UE_API bool AreAligningToActors() const;
 	/** For other systems to check if the Viewport World Interaction system currently has candidate actors selected */
-	bool HasCandidatesSelected();
+	UE_API bool HasCandidatesSelected();
 	/** If there are no currently selected candidates, use the currently selected actors as candidates. Otherwise, reset the candidates */
-	void SetSelectionAsCandidates();
+	UE_API void SetSelectionAsCandidates();
 
 	/** Gets the current delta time, so functions that don't get the delta time passed can still get it */
-	float GetCurrentDeltaTime() const;
+	UE_API float GetCurrentDeltaTime() const;
 
 	/** Getters and setters for whether or not to show the cursor on the viewport */
-	bool ShouldSuppressExistingCursor() const;
+	UE_API bool ShouldSuppressExistingCursor() const;
 	void SetShouldSuppressExistingCursor(const bool bInShouldSuppressCursor)
 	{
 		bShouldSuppressCursor = bInShouldSuppressCursor;
 	};
 
 	/** Getters and setters for whether or not to show the cursor on the viewport */
-	bool ShouldForceCursor() const;
-	void SetForceCursor(const bool bInShouldForceCursor);
+	UE_API bool ShouldForceCursor() const;
+	UE_API void SetForceCursor(const bool bInShouldForceCursor);
 
 	/** Gets the container for all the assets of ViewportInteraction. */
-	const UViewportInteractionAssetContainer& GetAssetContainer() const;
+	UE_API const UViewportInteractionAssetContainer& GetAssetContainer() const;
 
 	/** Static function to load the asset container */
-	static const UViewportInteractionAssetContainer* LoadAssetContainer();
+	static UE_API const UViewportInteractionAssetContainer* LoadAssetContainer();
 
 	/** Plays sound at location. */
-	void PlaySound(USoundBase* SoundBase, const FVector& InWorldLocation, const float InVolume = 1.0f);
+	UE_API void PlaySound(USoundBase* SoundBase, const FVector& InWorldLocation, const float InVolume = 1.0f);
 
 	/** Set if this world interaction is in VR. */
-	void SetInVR(const bool bInVR);
+	UE_API void SetInVR(const bool bInVR);
 
 	/** Get if this world interaction is in VR. */
-	bool IsInVR() const;
+	UE_API bool IsInVR() const;
 
-	FVector SnapLocation(const bool bLocalSpaceSnapping, const FVector& DesiredGizmoLocation, const FTransform &GizmoStartTransform, const FVector SnapGridBase, const bool bShouldConstrainMovement, const FVector AlignAxes);
+	UE_API FVector SnapLocation(const bool bLocalSpaceSnapping, const FVector& DesiredGizmoLocation, const FTransform &GizmoStartTransform, const FVector SnapGridBase, const bool bShouldConstrainMovement, const FVector AlignAxes);
 
 	/** Forces the VWI to fall back to standard desktop interactions */
-	void UseLegacyInteractions();
+	UE_API void UseLegacyInteractions();
 
 	/** Sets the VWI to use its own interactions */
-	void UseVWInteractions();
+	UE_API void UseVWInteractions();
 
 protected:
 
-	virtual void TransitionWorld(UWorld* NewWorld, EEditorWorldExtensionTransitionState TransitionState) override;
-	virtual void EnteredSimulateInEditor() override;
-	virtual void LeftSimulateInEditor(UWorld* SimulateWorld) override;
+	UE_API virtual void TransitionWorld(UWorld* NewWorld, EEditorWorldExtensionTransitionState TransitionState) override;
+	UE_API virtual void EnteredSimulateInEditor() override;
+	UE_API virtual void LeftSimulateInEditor(UWorld* SimulateWorld) override;
 
 private:
 
 	/** Called when the editor is closed */
-	void OnEditorClosed();
+	UE_API void OnEditorClosed();
 
 	/** Called every frame to update hover state */
-	void HoverTick( const float DeltaTime );
+	UE_API void HoverTick( const float DeltaTime );
 
 	/** Updates the interaction */
-	void InteractionTick( const float DeltaTime );
+	UE_API void InteractionTick( const float DeltaTime );
 
 	/** Called by the world interaction system when one of our components is dragged by the user.  If null is
 	    passed in then we'll treat it as dragging the whole object (rather than a specific axis/handle) */
-	void UpdateDragging(
+	UE_API void UpdateDragging(
 		const float DeltaTime,
 		bool& bIsFirstDragUpdate,
 		UViewportInteractor* Interactor,
@@ -426,7 +428,7 @@ private:
 		FVector& OutUnsnappedDraggedTo);
 
 	/** Given a drag delta from a starting point, contrains that delta based on a gizmo handle axis */
-	FVector ComputeConstrainedDragDeltaFromStart(
+	UE_API FVector ComputeConstrainedDragDeltaFromStart(
 		const bool bIsFirstDragUpdate,
 		const bool bOnPlane,
 		const TOptional<FTransformGizmoHandlePlacement> OptionalHandlePlacement,
@@ -442,10 +444,10 @@ private:
 
 	/** Called after dragging has finished with one or both hands and all inertial movement and interpolation has finished with
 	    the set of transformables we were working with.  Will also be called if the drag was interrupted. */
-	void FinishedMovingTransformables();
+	UE_API void FinishedMovingTransformables();
 
 	/** Checks to see if smooth snapping is currently enabled */
-	bool IsSmoothSnappingEnabled() const;
+	UE_API bool IsSmoothSnappingEnabled() const;
 
 	/** Returns the time since this was last entered */
 	inline FTimespan GetTimeSinceEntered() const
@@ -454,19 +456,19 @@ private:
 	}
 
 	/** Conditionally polls input from interactors, if we haven't done that already this frame */
-	void PollInputIfNeeded();
+	UE_API void PollInputIfNeeded();
 
 	/** Refreshes the transform gizmo to match the currently selected actor set */
-	void RefreshTransformGizmo( const bool bNewObjectsSelected );
+	UE_API void RefreshTransformGizmo( const bool bNewObjectsSelected );
 
 	/** Spawn new transform gizmo when refreshing and at start if there is none yet or the current one is different from the TransformGizmoClass */
-	void SpawnTransformGizmoIfNeeded();
+	UE_API void SpawnTransformGizmoIfNeeded();
 
 	/** Gets the inertie contribution of another interactor */
-	UViewportInteractor* GetOtherInteractorIntertiaContribute( UViewportInteractor* Interactor );
+	UE_API UViewportInteractor* GetOtherInteractorIntertiaContribute( UViewportInteractor* Interactor );
 
 	/** Destroys the actors */
-	void DestroyActors();
+	UE_API void DestroyActors();
 
 	/** Gets the snap grid mesh MID */
 	class UMaterialInstanceDynamic* GetSnapGridMID()
@@ -476,23 +478,23 @@ private:
 	}
 
 	/** Spawns the grid actor */
-	void SpawnGridMeshActor();
+	UE_API void SpawnGridMeshActor();
 
 	/** Average location of all the current transformables */
-	FVector CalculateAverageLocationOfTransformables();
+	UE_API FVector CalculateAverageLocationOfTransformables();
 
 	/** Location the transformable should snap to if aligning to other transformables*/
-	FVector FindTransformGizmoAlignPoint(const FTransform& GizmoStartTransform, const FTransform& DesiredGizmoTransform, const bool bShouldConstrainMovement, FVector ConstraintAxes);
-	void DrawBoxBrackets(const FBox InActor, const FTransform LocalToWorld, const FLinearColor BracketColor);
+	UE_API FVector FindTransformGizmoAlignPoint(const FTransform& GizmoStartTransform, const FTransform& DesiredGizmoTransform, const bool bShouldConstrainMovement, FVector ConstraintAxes);
+	UE_API void DrawBoxBrackets(const FBox InActor, const FTransform LocalToWorld, const FLinearColor BracketColor);
 
 	/** Calculate a new room transform location according to the world to meters scale */
-	void CompensateRoomTransformForWorldScale(FTransform& InOutRoomTransform, const float InNewWorldToMetersScale, const FVector& InRoomPivotLocation);
+	UE_API void CompensateRoomTransformForWorldScale(FTransform& InOutRoomTransform, const float InNewWorldToMetersScale, const FVector& InRoomPivotLocation);
 
 	/** If there is a transformable with velocity in simulate */
-	bool HasTransformableWithVelocityInSimulate() const;
+	UE_API bool HasTransformableWithVelocityInSimulate() const;
 
 	/** Get mode tools from the viewport client. */
-	class FEditorModeTools& GetModeTools() const;
+	UE_API class FEditorModeTools& GetModeTools() const;
 
 	//
 	// Colors
@@ -510,12 +512,12 @@ public:
 	};
 
 	/** Gets the color from color type */
-	FLinearColor GetColor(const EColors Color, const float Multiplier = 1.f) const;
+	UE_API FLinearColor GetColor(const EColors Color, const float Multiplier = 1.f) const;
 
 private:
 
 	/** The path of the asset container */
-	static const TCHAR* AssetContainerPath;
+	static UE_API const TCHAR* AssetContainerPath;
 
 	// All the colors for this mode
 	TArray<FLinearColor> Colors;
@@ -761,3 +763,5 @@ private:
 	/** Slate Input Processor */
 	TSharedPtr<class FViewportInteractionInputProcessor> InputProcessor;
 };
+
+#undef UE_API
